add an internal water weight channel and use it to more cleanly blend different water modifiers together

#Synced-CL 51695666

diff --git a/Engine/Plugins/Experimental/MeshPartitionWater/Source/MeshPartitionWater/Private/MeshPartitionLakeModifier.cpp b/Engine/Plugins/Experimental/MeshPartitionWater/Source/MeshPartitionWater/Private/MeshPartitionLakeModifier.cpp
--- a/Engine/Plugins/Experimental/MeshPartitionWater/Source/MeshPartitionWater/Private/MeshPartitionLakeModifier.cpp
+++ b/Engine/Plugins/Experimental/MeshPartitionWater/Source/MeshPartitionWater/Private/MeshPartitionLakeModifier.cpp
@@ -37,7 +37,7 @@ namespace MegaMeshLakeModifierLocals
 		// Generate a new random guid before submitting any code changes to the op
 		static FGuid GetCodeVersionKey()
 		{
-			static FGuid VersionKey(TEXT("2bc30c2d-e3bd-41b3-887c-8609dd475ab1"));
+			static FGuid VersionKey(TEXT("2472f4fa-a90b-459b-97a4-65f0779dd072"));
 			return VersionKey;
 		}
 
@@ -124,9 +124,14 @@ void MegaMeshLakeModifierLocals::FBackgroundOp::ApplyModifications(MeshPartition
 			return;
 		}
 
+		float InternalWaterWeight = InMeshView.GetVertexAttributeWeight(MeshPartition::ULakeModifier::InternalWaterWeightChannelName, VertexIndex);
+
 		const float NewVertexHeight = MeshPartition::ULakeModifier::CalculateVertexFalloffHeight(
+			InternalWaterWeight,
 			bIsInside, LakeHeight, Distance, ChannelDepth, MeshVertex.Z, WaterHeightmapSettings, CurveSettings);
 
+		InMeshView.SetVertexAttributeWeight(MeshPartition::ULakeModifier::InternalWaterWeightChannelName, VertexIndex, InternalWaterWeight);
+
 		for (const TPair<FName, FWaterBodyWeightmapSettings>& NameWt : WeightmapSettings)
 		{
 			const float ExistingWeight = InMeshView.GetVertexAttributeWeight(NameWt.Key, VertexIndex);
diff --git a/Engine/Plugins/Experimental/MeshPartitionWater/Source/MeshPartitionWater/Private/MeshPartitionRiverModifier.cpp b/Engine/Plugins/Experimental/MeshPartitionWater/Source/MeshPartitionWater/Private/MeshPartitionRiverModifier.cpp
--- a/Engine/Plugins/Experimental/MeshPartitionWater/Source/MeshPartitionWater/Private/MeshPartitionRiverModifier.cpp
+++ b/Engine/Plugins/Experimental/MeshPartitionWater/Source/MeshPartitionWater/Private/MeshPartitionRiverModifier.cpp
@@ -42,7 +42,7 @@ namespace MegaMeshRiverModifierLocals
 		// Generate a new random guid before submitting any code changes to the op
 		static FGuid GetCodeVersionKey()
 		{
-			static FGuid VersionKey(TEXT("4c892a39-58d9-4061-ab58-e99491aeda9b"));
+			static FGuid VersionKey(TEXT("e766a935-fe59-4997-a2e1-0d33f983ae89"));
 			return VersionKey;
 		}
 
@@ -256,7 +256,13 @@ void MegaMeshRiverModifierLocals::FBackgroundOp::ApplyModifications(MeshPartitio
 			return;
 		}
 
-		const float NewVertexHeight = URiverModifier::CalculateVertexFalloffHeight(bIsInside, RiverHeight, Distance, TargetDepth, MeshVertex.Z, WaterHeightmapSettings, CurveSettings);
+		float InternalWaterWeight = InMeshView.GetVertexAttributeWeight(MeshPartition::URiverModifier::InternalWaterWeightChannelName, VertexIndex);
+
+		const float NewVertexHeight = URiverModifier::CalculateVertexFalloffHeight(
+			InternalWaterWeight, 
+			bIsInside, RiverHeight, Distance, TargetDepth, MeshVertex.Z, WaterHeightmapSettings, CurveSettings);
+
+		InMeshView.SetVertexAttributeWeight(MeshPartition::URiverModifier::InternalWaterWeightChannelName, VertexIndex, InternalWaterWeight);
 
 		for (const TPair<FName, FWaterBodyWeightmapSettings>& NameWt : WeightmapSettings)
 		{
diff --git a/Engine/Plugins/Experimental/MeshPartitionWater/Source/MeshPartitionWater/Private/MeshPartitionWaterModifier.cpp b/Engine/Plugins/Experimental/MeshPartitionWater/Source/MeshPartitionWater/Private/MeshPartitionWaterModifier.cpp
--- a/Engine/Plugins/Experimental/MeshPartitionWater/Source/MeshPartitionWater/Private/MeshPartitionWaterModifier.cpp
+++ b/Engine/Plugins/Experimental/MeshPartitionWater/Source/MeshPartitionWater/Private/MeshPartitionWaterModifier.cpp
@@ -16,6 +16,9 @@
 
 namespace UE::MeshPartition
 {
+
+FName UWaterModifier::InternalWaterWeightChannelName("Water_InternalBlendChannel");
+
 UWaterModifier::UWaterModifier()
 {
 }
@@ -146,13 +149,20 @@ void UWaterModifier::OnWaterBodyChanged(const IWaterBrushActorInterface::FWaterB
 	}
 }
 
-float UWaterModifier::CalculateVertexFalloffHeight(bool bIsInside, float WaterHeight, float DistanceFromSpline, float TargetDepth, float MeshZ, const FWaterBodyHeightmapSettings& HeightmapSettings, const FWaterCurveSettings& CurveSettings)
+float UWaterModifier::CalculateVertexFalloffHeight(float& InOutInternalBlendWeight, bool bIsInside, float WaterHeight, float DistanceFromSpline, float TargetDepth, float MeshZ, const FWaterBodyHeightmapSettings& HeightmapSettings, const FWaterCurveSettings& CurveSettings)
 {
+	float TargetHeight = 0;
+	float TargetInternalWeight = 0;
+	float SourceInternalWeight = InOutInternalBlendWeight;
+	
+	// Find the target water height and weight depending on whether we're inside or outside of the water shape
 	if (bIsInside)
 	{
 		const double Alpha = FMath::Clamp((DistanceFromSpline - CurveSettings.ChannelEdgeOffset) / CurveSettings.CurveRampWidth, 0., 1.0);
 
-		return FMath::Lerp(WaterHeight, WaterHeight - TargetDepth, Alpha);
+		TargetHeight = FMath::Lerp(WaterHeight, WaterHeight - TargetDepth, Alpha);
+		// Inside the water, the weight is above 1
+		TargetInternalWeight = 1.f + Alpha;
 	}
 	else
 	{
@@ -168,8 +178,27 @@ float UWaterModifier::CalculateVertexFalloffHeight(bool bIsInside, float WaterHe
 
 		const double Alpha = FMath::Clamp((DistanceFromSpline - HeightmapSettings.FalloffSettings.EdgeOffset) / FalloffDistance, 0., 1.0);
 
-		return FMath::Lerp(WaterHeight, MeshZ, Alpha);
+		TargetHeight = WaterHeight;
+		// Outside the water, the weight falls off to 0
+		TargetInternalWeight = 1.f-Alpha;
 	}
+	
+	// start at the min of terrain and water target height
+	float ResultHeight = FMath::Min(TargetHeight, MeshZ);
+
+	// if we're in the outer-blend of old water, blend toward the new water's target height
+	if (SourceInternalWeight < 1.f)
+	{
+		ResultHeight = FMath::Lerp(ResultHeight, TargetHeight, 1.f - SourceInternalWeight);
+	}
+	// if we're in the outer-blend of new water, blend toward previous landscape height
+	if (TargetInternalWeight < 1.f)
+	{
+		ResultHeight = FMath::Lerp(ResultHeight, MeshZ, 1.f - TargetInternalWeight);
+	}
+
+	InOutInternalBlendWeight = FMath::Max(SourceInternalWeight, TargetInternalWeight);
+	return ResultHeight;
 }
 
 float UWaterModifier::CalculateVertexWeight(bool bIsInside, float DistanceFromSpline, const FWaterBodyWeightmapSettings& WeightmapSettings)
@@ -180,14 +209,15 @@ float UWaterModifier::CalculateVertexWeight(bool bIsInside, float DistanceFromSp
 
 void UWaterModifier::RegisterWaterWeightmaps(const TMap<FName, FWaterBodyWeightmapSettings>& WeightMaps, FInstanceInfo& Instance)
 {
+	Instance.UsedChannels.Emplace(InternalWaterWeightChannelName);
+	Instance.ReadViewComponents |= EMeshViewComponents::VertexAttributeWeight;
+	Instance.WriteViewComponents |= EMeshViewComponents::VertexAttributeWeight;
+
 	for (const TPair<FName, FWaterBodyWeightmapSettings>& WeightChannelEntry : WeightMaps)
 	{
 		if (!WeightChannelEntry.Key.IsNone())
 		{
 			Instance.UsedChannels.Emplace(WeightChannelEntry.Key);
-			
-			Instance.ReadViewComponents |= EMeshViewComponents::VertexAttributeWeight;
-			Instance.WriteViewComponents |= EMeshViewComponents::VertexAttributeWeight;
 		}
 	}
 }
diff --git a/Engine/Plugins/Experimental/MeshPartitionWater/Source/MeshPartitionWater/Public/MeshPartitionWaterModifier.h b/Engine/Plugins/Experimental/MeshPartitionWater/Source/MeshPartitionWater/Public/MeshPartitionWaterModifier.h
--- a/Engine/Plugins/Experimental/MeshPartitionWater/Source/MeshPartitionWater/Public/MeshPartitionWaterModifier.h
+++ b/Engine/Plugins/Experimental/MeshPartitionWater/Source/MeshPartitionWater/Public/MeshPartitionWaterModifier.h
@@ -48,7 +48,7 @@ protected:
 
 	UE_API void OnWaterBodyChanged(const IWaterBrushActorInterface::FWaterBrushActorChangedEventParams& OnBrushActorChanged);
 
-	static UE_API float CalculateVertexFalloffHeight(bool bIsInside, float WaterHeight, float DistanceFromSpline, float TargetHeight, float MeshZ, const FWaterBodyHeightmapSettings& HeightmapSettings, const FWaterCurveSettings& CurveSettings);
+	static UE_API float CalculateVertexFalloffHeight(float& InOutInternalBlendWeight, bool bIsInside, float WaterHeight, float DistanceFromSpline, float TargetHeight, float MeshZ, const FWaterBodyHeightmapSettings& HeightmapSettings, const FWaterCurveSettings& CurveSettings);
 	// helper to calculate basic weightmap alpha (not accounting for ModulationTexture)
 	static UE_API float CalculateVertexWeight(bool bIsInside, float DistanceFromSpline, const FWaterBodyWeightmapSettings& WeightmapSettings);
 	// helper to register water weightmaps
@@ -63,6 +63,9 @@ protected:
 	/* The maximum vertical distance this modifier will affect the terrain */
 	UPROPERTY(EditAnywhere, Category = Terrain)
 	double MaxZDistance = 10000.;
+
+	/* Name of the weight channel used to pass height blend info between overlapping water modifiers */
+	static UE_API FName InternalWaterWeightChannelName;
 };
 } // namespace UE::MeshPartition
 
