Updating Ground Truth Meshes for Remesh and Boolean Modifiers
https://epicgames.atlassian.net/browse/UE-368871

After change to bUseCompactSetAsDefault inside UBT BuildConfiguration, Boolean and Remesh modifiers' ground truth meshes need an update.
Also fixing the failing tests by setting default MeshPartitionDefinition for the tests.

#virtualized
#Synced-CL 52171699

diff --git a/Engine/Plugins/Experimental/MeshPartition/Source/MeshPartitionEditor/Private/Test/MeshPartitionBooleanModifierTests.cpp b/Engine/Plugins/Experimental/MeshPartition/Source/MeshPartitionEditor/Private/Test/MeshPartitionBooleanModifierTests.cpp
--- a/Engine/Plugins/Experimental/MeshPartition/Source/MeshPartitionEditor/Private/Test/MeshPartitionBooleanModifierTests.cpp
+++ b/Engine/Plugins/Experimental/MeshPartition/Source/MeshPartitionEditor/Private/Test/MeshPartitionBooleanModifierTests.cpp
@@ -5,6 +5,9 @@
 
 #include "CQTest.h"
 #include "Editor.h"
+#include "FileHelpers.h"
+#include "GameMapsSettings.h"
+#include "LevelEditorSubsystem.h"
 #include "MeshPartition.h"
 #include "ScopedTransaction.h"
 #include "Modifiers/MeshPartitionBooleanModifier.h"
@@ -15,14 +18,15 @@
 using namespace UE::MeshPartition;
 
 /**
- * Tests for MegaMesh Boolean Modifier
+ * Tests for MeshPartition Boolean Modifier
  */
 TEST_CLASS_WITH_FLAGS(BooleanModifierTests, "MeshPartition.Modifier", EAutomationTestFlags::EditorContext | EAutomationTestFlags::EngineFilter)
 {
 	// Test environment and core objects
 	TUniquePtr<FScopedEditorWorld> ScopedEditorWorld;
 	// CVars
-	TSharedPtr<FScopedTestEnvironment> TestEnvironment;
+	TOptional<TestUtils::FScopedCVarOverride> DisableDDCWrite;
+	TOptional<TestUtils::FScopedCVarOverride> DisableDDCRead;
 
 	UWorld* World = nullptr;
 	AMeshPartition* MeshPartition = nullptr;
@@ -59,6 +63,7 @@ TEST_CLASS_WITH_FLAGS(BooleanModifierTests, "MeshPartition.Modifier", EAutomatio
 		InitializationValues.CreateAISystem(false);
 		InitializationValues.AllowAudioPlayback(false);
 		InitializationValues.CreatePhysicsScene(true);
+		InitializationValues.CreateWorldPartition(true);
 
 #if AUTOMATED_TESTING_EXECUTE_IN_EDITOR_WORLD
 		World = GEditor->GetEditorWorldContext().World();
@@ -72,9 +77,8 @@ TEST_CLASS_WITH_FLAGS(BooleanModifierTests, "MeshPartition.Modifier", EAutomatio
 		ASSERT_THAT(IsNotNull(World, TEXT("World is not valid")));
 
 		// Need to disable DDC read/write, so that generated preview sections aren't written to DDC
-		TestEnvironment = FScopedTestEnvironment::Get();
-		TestEnvironment->SetConsoleVariableValue(TEXT("MegaMesh.Preview.EnableDDCWrite"), TEXT("0"));
-		TestEnvironment->SetConsoleVariableValue(TEXT("MegaMesh.Preview.EnableDDCRead"), TEXT("0"));
+		DisableDDCWrite.Emplace(TEXT("MegaMesh.Preview.EnableDDCWrite"), TEXT("0"));
+		DisableDDCRead.Emplace(TEXT("MegaMesh.Preview.EnableDDCRead"), TEXT("0"));
 
 		MeshPartition = TestUtils::CreateTestMesh(
 			World, MeshWidth, MeshHeight, MeshWidthVertexCount, MeshHeightVertexCount
@@ -86,6 +90,10 @@ TEST_CLASS_WITH_FLAGS(BooleanModifierTests, "MeshPartition.Modifier", EAutomatio
 
 		// Pausing the transformer pipeline during the test to avoid potential race condition
 		MeshPartitionEditorComponent->SetPauseTransformerPipeline(true);
+
+		// Load WorldPartition Region
+		const bool bSuccessfullyLoadedWorldPartition = TestUtils::LoadWorldPartitionRegion(World);
+		ASSERT_THAT(IsTrue(bSuccessfullyLoadedWorldPartition, TEXT("Failed to load WorldPartition region")));
 	}
 
 	// UE-TC-22590
@@ -141,7 +149,7 @@ TEST_CLASS_WITH_FLAGS(BooleanModifierTests, "MeshPartition.Modifier", EAutomatio
 				})
 			.Until(TEXT("Wait for preview build completion"), [this]()
 				{
-					return !MeshPartitionEditorComponent->IsModifierParticipatingInActivePreviewSectionBuild(BooleanModifier);
+					return !MeshPartitionEditorComponent->IsAnyPreviewSectionBuildActive();
 				})
 			.Then(TEXT("Validate preview mesh against ground truth"), [this]()
 				{
@@ -508,6 +516,10 @@ TEST_CLASS_WITH_FLAGS(BooleanModifierTests, "MeshPartition.Modifier", EAutomatio
 	TEST_METHOD(Verify_Boolean_Modifier_Affects_Assigned_MeshTerrain_Only)
 	{
 		TestCommandBuilder
+			.Until(TEXT("Wait for preview build completion"), [this]()
+				{
+					return !MeshPartitionEditorComponent->IsAnyPreviewSectionBuildActive();
+				})
 			.Do(TEXT("Create second MeshPartition and move it"), [this]()
 				{
 					MeshPartition2 = TestUtils::CreateTestMesh(
@@ -576,6 +588,10 @@ TEST_CLASS_WITH_FLAGS(BooleanModifierTests, "MeshPartition.Modifier", EAutomatio
 	TEST_METHOD(Verify_Boolean_Modifier_AffectedMeshPartition_Reassignment_Affects_New_Terrain_Only)
 	{
 		TestCommandBuilder
+			.Until(TEXT("Wait for preview build completion"), [this]()
+				{
+					return !MeshPartitionEditorComponent->IsAnyPreviewSectionBuildActive();
+				})
 			.Do(TEXT("Create second MeshPartition and move it"), [this]()
 				{
 					MeshPartition2 = TestUtils::CreateTestMesh(
@@ -645,6 +661,10 @@ TEST_CLASS_WITH_FLAGS(BooleanModifierTests, "MeshPartition.Modifier", EAutomatio
 	TEST_METHOD(Verify_Boolean_Modifier_AffectedMeshPartition_Assignment_UndoRedo_RevertsAndReapplies)
 	{
 		TestCommandBuilder
+			.Until(TEXT("Wait for preview build completion"), [this]()
+				{
+					return !MeshPartitionEditorComponent->IsAnyPreviewSectionBuildActive();
+				})
 			.Do(TEXT("Create second MeshPartition and move it"), [this]()
 				{
 					MeshPartition2 = TestUtils::CreateTestMesh(
diff --git a/Engine/Plugins/Experimental/MeshPartition/Source/MeshPartitionEditor/Private/Test/MeshPartitionHeightmapImportTests.cpp b/Engine/Plugins/Experimental/MeshPartition/Source/MeshPartitionEditor/Private/Test/MeshPartitionHeightmapImportTests.cpp
--- a/Engine/Plugins/Experimental/MeshPartition/Source/MeshPartitionEditor/Private/Test/MeshPartitionHeightmapImportTests.cpp
+++ b/Engine/Plugins/Experimental/MeshPartition/Source/MeshPartitionEditor/Private/Test/MeshPartitionHeightmapImportTests.cpp
@@ -22,7 +22,7 @@
 using namespace UE::MeshPartition;
 
 /**
- * Tests for MegaMesh Heightmap Import
+ * Tests for MeshPartition Heightmap Import
  */
 TEST_CLASS_WITH_FLAGS(HeightmapImport, "MeshPartition", EAutomationTestFlags::EditorContext | EAutomationTestFlags::EngineFilter)
 {
@@ -30,12 +30,13 @@ private:
 	// Test environment and core objects
 	TUniquePtr<FScopedEditorWorld> ScopedEditorWorld;
 	// CVars
-	TSharedPtr<FScopedTestEnvironment> TestEnvironment;
+	TOptional<TestUtils::FScopedCVarOverride> DisableDDCWrite;
+	TOptional<TestUtils::FScopedCVarOverride> DisableDDCRead;
 
 	// Core test objects
 	UWorld* World = nullptr;
-	AMeshPartition* MegaMesh = nullptr;
-	UMeshPartitionEditorComponent* MegaMeshEditorComponent = nullptr;
+	AMeshPartition* MeshPartition = nullptr;
+	UMeshPartitionEditorComponent* MeshPartitionEditorComponent = nullptr;
 
 	// Test data
 	TArray<TObjectPtr<UStaticMesh>> GroundTruthBaseMeshes;
@@ -72,6 +73,7 @@ public:
 		InitializationValues.CreateAISystem(false);
 		InitializationValues.AllowAudioPlayback(false);
 		InitializationValues.CreatePhysicsScene(true);
+		InitializationValues.CreateWorldPartition(true);
 
 		// Spawn world
 		ScopedEditorWorld = FAutomationEditorCommonUtils::CreateScopedEditorWorld(LevelPath, InitializationValues);
@@ -81,20 +83,22 @@ public:
 		ASSERT_THAT(IsNotNull(World, TEXT("Failed to create valid test world")));
 
 		// Need to disable DDC read/write, so that generated preview sections aren't written to DDC
-		TestEnvironment = FScopedTestEnvironment::Get();
-		ASSERT_THAT(IsNotNull(TestEnvironment.Get(), TEXT("Failed to get scoped test environment")));
-		TestEnvironment->SetConsoleVariableValue(TEXT("MegaMesh.Preview.EnableDDCWrite"), TEXT("0"));
-		TestEnvironment->SetConsoleVariableValue(TEXT("MegaMesh.Preview.EnableDDCRead"), TEXT("0"));
+		DisableDDCWrite.Emplace(TEXT("MegaMesh.Preview.EnableDDCWrite"), TEXT("0"));
+		DisableDDCRead.Emplace(TEXT("MegaMesh.Preview.EnableDDCRead"), TEXT("0"));
+
+		// Load WorldPartition Region
+		const bool bSuccessfullyLoadedWorldPartition = TestUtils::LoadWorldPartitionRegion(World);
+		ASSERT_THAT(IsTrue(bSuccessfullyLoadedWorldPartition, TEXT("Failed to load WorldPartition region")));
 	}
 
 	/**
-	 * Verifies that heightmap import correctly generates the expected MegaMesh structure
+	 * Verifies that heightmap import correctly generates the expected MeshPartition structure
 	 * with proper modifiers and preview sections that match ground truth meshes.
 	 */
 	TEST_METHOD(Verify_HeightmapImport_Creates_Expected_Modifiers_And_PreviewSections)
 	{
 		TestCommandBuilder
-			.Do(TEXT("Import heightmap using MegaMeshHeightmapImporter"), [this]()
+			.Do(TEXT("Import heightmap using MeshPartitionHeightmapImporter"), [this]()
 				{
 					const FString PluginDirectory = IPluginManager::Get().FindPlugin(TEXT("MeshPartition"))->GetBaseDir();
 
@@ -117,17 +121,17 @@ public:
 					FPackageSourceControlHelper* SourceControlHelper = nullptr;
 					Importer.Import(SourceControlHelper);
 
-					MegaMesh = Importer.GetMegaMesh();
-					ASSERT_THAT(IsNotNull(MegaMesh, TEXT("Failed to get MegaMesh from imported heightmap")));
+					MeshPartition = Importer.GetMegaMesh();
+					ASSERT_THAT(IsNotNull(MeshPartition, TEXT("Failed to get MeshPartition from imported heightmap")));
 				})
 			.Then(TEXT("Build preview sections and configure editor component"), [this]()
 				{
-					MegaMeshEditorComponent = Cast<UMeshPartitionEditorComponent>(MegaMesh->GetMeshPartitionComponent());
-					ASSERT_THAT(IsNotNull(MegaMeshEditorComponent,
-						TEXT("Failed to get MegaMeshEditorComponent from imported MegaMesh")));
+					MeshPartitionEditorComponent = Cast<UMeshPartitionEditorComponent>(MeshPartition->GetMeshPartitionComponent());
+					ASSERT_THAT(IsNotNull(MeshPartitionEditorComponent,
+						TEXT("Failed to get MeshPartitionEditorComponent from imported MeshPartition")));
 
 					// Pausing Transformer pipeline during the test to avoid potential race condition
-					MegaMeshEditorComponent->SetPauseTransformerPipeline(true);
+					MeshPartitionEditorComponent->SetPauseTransformerPipeline(true);
 
 					// Build preview sections with infinite bounds
 					TArray InfiniteBounds = {
@@ -135,17 +139,17 @@ public:
 							FVector3d(HALF_WORLD_MAX, HALF_WORLD_MAX, HALF_WORLD_MAX))
 					};
 
-					MegaMeshEditorComponent->BuildMegaMeshPreviewSections(InfiniteBounds);
-					MegaMeshEditorComponent->Update();
+					MeshPartitionEditorComponent->BuildMegaMeshPreviewSections(InfiniteBounds);
+					MeshPartitionEditorComponent->Update();
 				})
 			.Until(TEXT("Wait for preview section build to complete"), [this]()
 				{
-					return !MegaMeshEditorComponent->IsAnyPreviewSectionBuildActive();
+					return !MeshPartitionEditorComponent->IsAnyPreviewSectionBuildActive();
 				})
 			.Then(TEXT("Validate expected number of base modifiers were created"), [this]()
 				{
 					// Get all base modifiers
-					Modifiers = MegaMeshEditorComponent->GetModifiersFiltered([](const UModifierComponent* Modifier)
+					Modifiers = MeshPartitionEditorComponent->GetModifiersFiltered([](const UModifierComponent* Modifier)
 					{
 						return Modifier->IsBase() && Cast<UMeshProviderModifier>(Modifier);
 					});
@@ -196,7 +200,7 @@ public:
 
 					// Validate preview sections were created
 					int32 ActualPreviewSectionCount = 0;
-					MegaMeshEditorComponent->ForAllPreviewSections([&ActualPreviewSectionCount](const APreviewSection* PreviewSection)
+					MeshPartitionEditorComponent->ForAllPreviewSections([&ActualPreviewSectionCount](const APreviewSection* PreviewSection)
 						{
 							if (IsValid(PreviewSection))
 							{
diff --git a/Engine/Plugins/Experimental/MeshPartition/Source/MeshPartitionEditor/Private/Test/MeshPartitionModifierLifetimeTests.cpp b/Engine/Plugins/Experimental/MeshPartition/Source/MeshPartitionEditor/Private/Test/MeshPartitionModifierLifetimeTests.cpp
--- a/Engine/Plugins/Experimental/MeshPartition/Source/MeshPartitionEditor/Private/Test/MeshPartitionModifierLifetimeTests.cpp
+++ b/Engine/Plugins/Experimental/MeshPartition/Source/MeshPartitionEditor/Private/Test/MeshPartitionModifierLifetimeTests.cpp
@@ -12,10 +12,11 @@
 #include "Engine/World.h"
 #include "MeshPartitionTestModifier.h"
 #include "CQTest.h"
-#include "MeshPartitionEditorSubsystem.h"
 #include "Tests/AutomationCommon.h"
 #include "Tests/AutomationEditorCommon.h"
 
+using namespace UE::MeshPartition;
+
 /**
  * Tests for MeshPartition Modifier lifecycle initialization and cleanup
  */
@@ -37,7 +38,8 @@ TEST_CLASS_WITH_FLAGS(Initialization, "MeshPartition.Modifier", EAutomationTestF
 	const float TickDelta = 0.016f;
 
 	// CVars
-	TSharedPtr<FScopedTestEnvironment> TestEnvironment;
+	TOptional<TestUtils::FScopedCVarOverride> DisableDDCWrite;
+	TOptional<TestUtils::FScopedCVarOverride> DisableDDCRead;
 
 	BEFORE_EACH()
 	{
@@ -50,6 +52,7 @@ TEST_CLASS_WITH_FLAGS(Initialization, "MeshPartition.Modifier", EAutomationTestF
 		InitializationValues.CreateAISystem(false);
 		InitializationValues.AllowAudioPlayback(false);
 		InitializationValues.CreatePhysicsScene(true);
+		InitializationValues.CreateWorldPartition(true);
 
 		// Spawn world
 		const FString LevelPath = TEXT("/MeshPartition/UnitTests/MeshPartitionUnitTests");
@@ -59,9 +62,12 @@ TEST_CLASS_WITH_FLAGS(Initialization, "MeshPartition.Modifier", EAutomationTestF
 		ASSERT_THAT(IsNotNull(World, TEXT("Failed to create valid test world.")));
 
 		// Need to disable DDC read/write, so that generated preview sections aren't written to DDC
-		TestEnvironment = FScopedTestEnvironment::Get();
-		TestEnvironment->SetConsoleVariableValue(TEXT("MegaMesh.Preview.EnableDDCWrite"), TEXT("0"));
-		TestEnvironment->SetConsoleVariableValue(TEXT("MegaMesh.Preview.EnableDDCRead"), TEXT("0"));
+		DisableDDCWrite.Emplace(TEXT("MegaMesh.Preview.EnableDDCWrite"), TEXT("0"));
+		DisableDDCRead.Emplace(TEXT("MegaMesh.Preview.EnableDDCRead"), TEXT("0"));
+
+		// Load WorldPartition Region
+		const bool bSuccessfullyLoadedWorldPartition = UE::MeshPartition::TestUtils::LoadWorldPartitionRegion(World);
+		ASSERT_THAT(IsTrue(bSuccessfullyLoadedWorldPartition, TEXT("Failed to load WorldPartition region")));
 	}
 
 	/**
diff --git a/Engine/Plugins/Experimental/MeshPartition/Source/MeshPartitionEditor/Private/Test/MeshPartitionPatchTests.cpp b/Engine/Plugins/Experimental/MeshPartition/Source/MeshPartitionEditor/Private/Test/MeshPartitionPatchTests.cpp
--- a/Engine/Plugins/Experimental/MeshPartition/Source/MeshPartitionEditor/Private/Test/MeshPartitionPatchTests.cpp
+++ b/Engine/Plugins/Experimental/MeshPartition/Source/MeshPartitionEditor/Private/Test/MeshPartitionPatchTests.cpp
@@ -16,7 +16,7 @@
 using namespace UE::MeshPartition;
 
 /**
- * Tests for MegaMesh Patch modifier
+ * Tests for MeshPartition Patch modifier
  */
 TEST_CLASS_WITH_FLAGS(Patch, "MeshPartition.Modifier", EAutomationTestFlags::EditorContext | EAutomationTestFlags::EngineFilter)
 {
@@ -24,12 +24,13 @@ private:
 	// Test environment and core objects
 	TUniquePtr<FScopedEditorWorld> ScopedEditorWorld;
 	// CVars
-	TSharedPtr<FScopedTestEnvironment> TestEnvironment;
+	TOptional<TestUtils::FScopedCVarOverride> DisableDDCWrite;
+	TOptional<TestUtils::FScopedCVarOverride> DisableDDCRead;
 
 	// Core test objects
 	UWorld* World = nullptr;
-	AMeshPartition* MegaMesh = nullptr;
-	UMeshPartitionEditorComponent* MegaMeshEditorComponent = nullptr;
+	AMeshPartition* MeshPartition = nullptr;
+	UMeshPartitionEditorComponent* MeshPartitionEditorComponent = nullptr;
 
 	// Test constants
 	static constexpr int32 MeshWidth = 20000;
@@ -38,6 +39,7 @@ private:
 	static constexpr int32 MeshHeightVertexCount = 100;
 	static constexpr int32 ExpectedPreviewSectionCount = 1;
 	static inline const FString LevelPath = TEXT("/MeshPartition/UnitTests/MeshPartitionUnitTests");
+	const FString ReferencePath = TEXT("/MeshPartition/UnitTests/PatchModifierTests");
 
 public:
 	BEFORE_EACH()
@@ -51,6 +53,7 @@ public:
 		InitializationValues.CreateAISystem(false);
 		InitializationValues.AllowAudioPlayback(false);
 		InitializationValues.CreatePhysicsScene(true);
+		InitializationValues.CreateWorldPartition(true);
 
 		// Spawn world
 		ScopedEditorWorld = FAutomationEditorCommonUtils::CreateScopedEditorWorld(LevelPath, InitializationValues);
@@ -60,11 +63,12 @@ public:
 		ASSERT_THAT(IsNotNull(World, TEXT("Failed to create valid test world")));
 
 		// Need to disable DDC read/write, so that generated preview sections aren't written to DDC
-		TestEnvironment = FScopedTestEnvironment::Get();
-		ASSERT_THAT(IsNotNull(TestEnvironment.Get(), TEXT("Failed to get scoped test environment")));
+		DisableDDCWrite.Emplace(TEXT("MegaMesh.Preview.EnableDDCWrite"), TEXT("0"));
+		DisableDDCRead.Emplace(TEXT("MegaMesh.Preview.EnableDDCRead"), TEXT("0"));
 
-		TestEnvironment->SetConsoleVariableValue(TEXT("MegaMesh.Preview.EnableDDCWrite"), TEXT("0"));
-		TestEnvironment->SetConsoleVariableValue(TEXT("MegaMesh.Preview.EnableDDCRead"), TEXT("0"));
+		// Load WorldPartition Region
+		const bool bSuccessfullyLoadedWorldPartition = TestUtils::LoadWorldPartitionRegion(World);
+		ASSERT_THAT(IsTrue(bSuccessfullyLoadedWorldPartition, TEXT("Failed to load WorldPartition region")));
 	}
 
 	/**
@@ -74,17 +78,21 @@ public:
 	TEST_METHOD(Verify_PatchModifier_Priority_System_Applies_Higher_Priority_First)
 	{
 		TestCommandBuilder
-			.Do(TEXT("Create MegaMesh with test dimensions"), [this]()
+			.Do(TEXT("Create MeshPartition with test dimensions"), [this]()
 			{
-				// Create MegaMesh
-				MegaMesh = TestUtils::CreateTestMesh(
+				// Create MeshPartition
+				MeshPartition = TestUtils::CreateTestMesh(
 					World, MeshWidth, MeshHeight, MeshWidthVertexCount, MeshHeightVertexCount);
-				ASSERT_THAT(IsNotNull(MegaMesh, TEXT("Failed to create MegaMesh")));
+				ASSERT_THAT(IsNotNull(MeshPartition, TEXT("Failed to create MeshPartition")));
 
-				MegaMeshEditorComponent = Cast<UMeshPartitionEditorComponent>(MegaMesh->GetMeshPartitionComponent());
-				ASSERT_THAT(IsNotNull(MegaMeshEditorComponent,
-					TEXT("Failed to get MegaMeshEditorComponent from MegaMesh")));
+				MeshPartitionEditorComponent = Cast<UMeshPartitionEditorComponent>(MeshPartition->GetMeshPartitionComponent());
+				ASSERT_THAT(IsNotNull(MeshPartitionEditorComponent,
+					TEXT("Failed to get MeshPartitionEditorComponent from MeshPartition")));
 			})
+			.Until(TEXT("Wait for preview section build to complete"), [this]()
+				{
+					return !MeshPartitionEditorComponent->IsAnyPreviewSectionBuildActive();
+				})
 			.Then(TEXT("Create overlapping patch modifiers with different priorities"), [this]()
 			{
 				// Create high priority patch modifier
@@ -93,26 +101,26 @@ public:
 					constexpr double HighPriority = 1.0;
 					AActor* HighPriorityPatchActor = World->SpawnActor<AActor>(AActor::StaticClass(),
 						FTransform::Identity);
-					ASSERT_THAT(IsNotNull(HighPriorityPatchActor, TEXT("Failed to spawn high priority MegaMeshModifier")));
+					ASSERT_THAT(IsNotNull(HighPriorityPatchActor, TEXT("Failed to spawn high priority MeshPartitionModifier")));
 
-					TestUtils::CreatePatchModifier(HighPriorityPatchActor, MegaMesh, ModifierWorldLocation, HighPriority);
+					TestUtils::CreatePatchModifier(HighPriorityPatchActor, MeshPartition, ModifierWorldLocation, HighPriority);
 				}
 				// Create low priority patch modifier
 				{
 					const FVector ModifierWorldLocation = FVector(0.0, 1000.0, 2500.0);
 					AActor* LowPriorityPatchActor = World->SpawnActor<AActor>(AActor::StaticClass(),
 						FTransform::Identity);
-					ASSERT_THAT(IsNotNull(LowPriorityPatchActor, TEXT("Failed to spawn low priority MegaMeshModifier")));
+					ASSERT_THAT(IsNotNull(LowPriorityPatchActor, TEXT("Failed to spawn low priority MeshPartitionModifier")));
 
-					TestUtils::CreatePatchModifier(LowPriorityPatchActor, MegaMesh, ModifierWorldLocation);
+					TestUtils::CreatePatchModifier(LowPriorityPatchActor, MeshPartition, ModifierWorldLocation);
 				}
 
 				// Notify the editor component that modifiers have been assigned
-				MegaMeshEditorComponent->OnModifierAssigned();
+				MeshPartitionEditorComponent->OnModifierAssigned();
 			})
 			.Until(TEXT("Wait for preview section build to complete"), [this]()
 			{
-				return !MegaMeshEditorComponent->IsAnyPreviewSectionBuildActive();
+				return !MeshPartitionEditorComponent->IsAnyPreviewSectionBuildActive();
 			})
 			.Then(TEXT("Validate preview section matches expected mesh"), [this]()
 			{
@@ -120,7 +128,7 @@ public:
 				int32 PreviewSectionCount = 0;
 
 				// Count and capture preview sections
-				MegaMeshEditorComponent->ForAllPreviewSections([&](const APreviewSection* InPreviewSection)
+				MeshPartitionEditorComponent->ForAllPreviewSections([&](const APreviewSection* InPreviewSection)
 				{
 					PreviewSection = InPreviewSection;
 					++PreviewSectionCount;
@@ -130,22 +138,11 @@ public:
 				// Validate preview section count
 				ASSERT_THAT(AreEqual(ExpectedPreviewSectionCount, PreviewSectionCount,
 					*FString::Printf(TEXT("PreviewSectionCount should be %d"), ExpectedPreviewSectionCount)));
-				ASSERT_THAT(IsNotNull(PreviewSection, TEXT("MegaMesh PreviewSection is null")));
+				ASSERT_THAT(IsNotNull(PreviewSection, TEXT("MeshPartition PreviewSection is null")));
 
 				// Validate preview mesh matches ground truth
 				const FString GroundTruthAsset = TEXT("Verify_PatchModifier_Priority_System_Applies_Higher_Priority_First_GroundTruth");
-				const FString GroundTruthPath = FString::Printf(
-					TEXT("/MeshPartition/UnitTests/PatchModifierTests/%s.%s"),
-					*GroundTruthAsset,
-					*GroundTruthAsset
-				);
-				FString ErrorMessage;
-				const bool bValid = TestUtils::ValidatePreviewMeshAgainstGroundTruth(
-					MegaMeshEditorComponent,
-					{ GroundTruthPath },
-					ErrorMessage);
-
-				ASSERT_THAT(IsTrue(bValid, *ErrorMessage));
+				CompareToReferenceAsset(GroundTruthAsset);
 			});
 	}
 
@@ -154,6 +151,27 @@ public:
 		ScopedEditorWorld.Reset();
 		CollectGarbage(GARBAGE_COLLECTION_KEEPFLAGS);
 	}
+
+private:
+	void CompareToReferenceAsset(const FString& TestName)
+	{
+		const FString ReferencePackagePath = ReferencePath / TestName;
+		const FString ReferenceAssetPath = ReferencePackagePath + "." + TestName;
+
+		if (TestUtils::ShouldUpdateReferenceAssets())
+		{
+			[[maybe_unused]] const bool bSuccess =
+				TestUtils::CreateOrUpdateReference(MeshPartitionEditorComponent, ReferencePackagePath, { TestName }, TestName);
+		}
+
+		FString ErrorMessage;
+		bool bValid = TestUtils::ValidatePreviewMeshAgainstGroundTruth(
+			MeshPartitionEditorComponent,
+			{ ReferenceAssetPath },
+			ErrorMessage);
+
+		ASSERT_THAT(IsTrue(bValid, *ErrorMessage));
+	}
 };
 
 #endif // WITH_EDITOR
diff --git a/Engine/Plugins/Experimental/MeshPartition/Source/MeshPartitionEditor/Private/Test/MeshPartitionPreviewBuildTests.cpp b/Engine/Plugins/Experimental/MeshPartition/Source/MeshPartitionEditor/Private/Test/MeshPartitionPreviewBuildTests.cpp
--- a/Engine/Plugins/Experimental/MeshPartition/Source/MeshPartitionEditor/Private/Test/MeshPartitionPreviewBuildTests.cpp
+++ b/Engine/Plugins/Experimental/MeshPartition/Source/MeshPartitionEditor/Private/Test/MeshPartitionPreviewBuildTests.cpp
@@ -15,17 +15,20 @@
 #include "Tests/AutomationCommon.h"
 #include "Tests/AutomationEditorCommon.h"
 
+using namespace UE::MeshPartition;
+
 TEST_CLASS_WITH_FLAGS(PreviewSection, "MeshPartition", EAutomationTestFlags::EditorContext | EAutomationTestFlags::EngineFilter)
 {
 	// Test environment and core objects
 	TUniquePtr<FScopedEditorWorld> ScopedEditorWorld;
 	UWorld* World = nullptr;
-	UE::MeshPartition::AMeshPartition* MegaMesh = nullptr;
+	UE::MeshPartition::AMeshPartition* MeshPartition = nullptr;
 	UE::MeshPartition::UMeshPartitionEditorComponent* MeshPartitionEditorComponent = nullptr;
 	TArray<TObjectPtr<UStaticMesh>> GroundTruthPreviewMeshes;
 
 	// CVars
-	TSharedPtr<FScopedTestEnvironment> TestEnvironment;
+	TOptional<TestUtils::FScopedCVarOverride> DisableDDCWrite;
+	TOptional<TestUtils::FScopedCVarOverride> DisableDDCRead;
 
 	BEFORE_EACH()
 	{
@@ -49,28 +52,31 @@ TEST_CLASS_WITH_FLAGS(PreviewSection, "MeshPartition", EAutomationTestFlags::Edi
 		World = ScopedEditorWorld->GetWorld();
 		ASSERT_THAT(IsNotNull(World, TEXT("Failed to create valid test world.")));
 
-		MegaMesh = World->SpawnActor<AMeshPartition>(AMeshPartition::StaticClass(), FTransform::Identity);
-		ASSERT_THAT(IsNotNull(MegaMesh, TEXT("Failed to spawn MegaMesh actor.")));
+		MeshPartition = World->SpawnActor<AMeshPartition>(AMeshPartition::StaticClass(), FTransform::Identity);
+		ASSERT_THAT(IsNotNull(MeshPartition, TEXT("Failed to spawn MeshPartition actor.")));
 
 		MeshPartitionEditorComponent = NewObject<UMeshPartitionEditorComponent>(
-			MegaMesh, UMeshPartitionEditorComponent::StaticClass(), TEXT("MeshPartitionEditorComponent"));
-		ASSERT_THAT(IsNotNull(MeshPartitionEditorComponent, TEXT("Failed to create editor component for MegaMesh.")));
+			MeshPartition, UMeshPartitionEditorComponent::StaticClass(), TEXT("MeshPartitionEditorComponent"));
+		ASSERT_THAT(IsNotNull(MeshPartitionEditorComponent, TEXT("Failed to create editor component for MeshPartition.")));
 
 		// Set up component hierarchy
-		MegaMesh->SetMeshPartitionComponent(MeshPartitionEditorComponent);
+		MeshPartition->SetMeshPartitionComponent(MeshPartitionEditorComponent);
 		MeshPartitionEditorComponent->RegisterComponent();
 
 		// Pausing Transformer pipeline during the test to avoid potential race condition
 		MeshPartitionEditorComponent->SetPauseTransformerPipeline(true);
 
 		// Need to disable DDC read/write, so that generated preview sections aren't written to DDC
-		TestEnvironment = FScopedTestEnvironment::Get();
-		TestEnvironment->SetConsoleVariableValue(TEXT("MegaMesh.Preview.EnableDDCWrite"), TEXT("0"));
-		TestEnvironment->SetConsoleVariableValue(TEXT("MegaMesh.Preview.EnableDDCRead"), TEXT("0"));
+		DisableDDCWrite.Emplace(TEXT("MegaMesh.Preview.EnableDDCWrite"), TEXT("0"));
+		DisableDDCRead.Emplace(TEXT("MegaMesh.Preview.EnableDDCRead"), TEXT("0"));
+
+		// Load WorldPartition Region
+		const bool bSuccessfullyLoadedWorldPartition = TestUtils::LoadWorldPartitionRegion(World);
+		ASSERT_THAT(IsTrue(bSuccessfullyLoadedWorldPartition, TEXT("Failed to load WorldPartition region")));
 	}
 
 	/**
-	 * Verifies that MegaMesh correctly creates preview sections based on complexity settings,
+	 * Verifies that MeshPartition correctly creates preview sections based on complexity settings,
 	 * merging low-complexity sections and separating high-complexity ones
 	 */
 	TEST_METHOD(Verify_PreviewSection_Complexity_Threshold_Behavior)
@@ -85,7 +91,7 @@ TEST_CLASS_WITH_FLAGS(PreviewSection, "MeshPartition", EAutomationTestFlags::Edi
 				UE::MeshPartition::FCommonBuildVariant Variant = Definition->GetPreviewSectionBuildVariant();
 				Variant.MaxSectionComplexity = 50;
 				Definition->SetPreviewSectionBuildVariant(Variant);
-				MegaMesh->SetMeshPartitionDefinition(Definition);
+				MeshPartition->SetMeshPartitionDefinition(Definition);
 
 				TArray<FBox> Bounds;
 				// Create two base modifiers with two triangle quads
@@ -119,7 +125,7 @@ TEST_CLASS_WITH_FLAGS(PreviewSection, "MeshPartition", EAutomationTestFlags::Edi
 	}
 
 	/**
-	 * Verifies that MegaMesh correctly applies transforms when creating preview sections,
+	 * Verifies that MeshPartition correctly applies transforms when creating preview sections,
 	 * ensuring the generated mesh matches the expected ground truth
 	 */
 	TEST_METHOD(Verify_PreviewSection_Preserves_Transforms)
@@ -169,7 +175,7 @@ TEST_CLASS_WITH_FLAGS(PreviewSection, "MeshPartition", EAutomationTestFlags::Edi
 	}
 
 	/**
-	 * Verifies that MegaMesh correctly merges overlapping geometry into a single preview section
+	 * Verifies that MeshPartition correctly merges overlapping geometry into a single preview section
 	 */
 	TEST_METHOD(Verify_PreviewSection_Merges_Overlapping_Geometry)
 	{
@@ -218,7 +224,7 @@ TEST_CLASS_WITH_FLAGS(PreviewSection, "MeshPartition", EAutomationTestFlags::Edi
 	}
 
 	/**
-	 * Verifies that MegaMesh correctly maintains separate preview sections for non-overlapping geometry
+	 * Verifies that MeshPartition correctly maintains separate preview sections for non-overlapping geometry
 	 */
 	TEST_METHOD(Verify_PreviewSection_Preserves_NonOnverlapping_Geometry)
 	{
@@ -265,7 +271,7 @@ TEST_CLASS_WITH_FLAGS(PreviewSection, "MeshPartition", EAutomationTestFlags::Edi
 	}
 
 	/**
-	 * Verifies that MegaMesh correctly handles aggregated bounds from multiple modifiers
+	 * Verifies that MeshPartition correctly handles aggregated bounds from multiple modifiers
 	 */
 	TEST_METHOD(Verify_PreviewSection_Handles_Aggregated_Bounds)
 	{
diff --git a/Engine/Plugins/Experimental/MeshPartition/Source/MeshPartitionEditor/Private/Test/MeshPartitionRemeshModifierTests.cpp b/Engine/Plugins/Experimental/MeshPartition/Source/MeshPartitionEditor/Private/Test/MeshPartitionRemeshModifierTests.cpp
--- a/Engine/Plugins/Experimental/MeshPartition/Source/MeshPartitionEditor/Private/Test/MeshPartitionRemeshModifierTests.cpp
+++ b/Engine/Plugins/Experimental/MeshPartition/Source/MeshPartitionEditor/Private/Test/MeshPartitionRemeshModifierTests.cpp
@@ -7,7 +7,6 @@
 #include "Editor.h"
 #include "MeshPartition.h"
 #include "MeshPartitionHeightmapImporter.h"
-#include "MeshPartitionPreviewSection.h"
 #include "Interfaces/IPluginManager.h"
 #include "Modifiers/MeshPartitionRemeshModifier.h"
 #include "Test/MeshPartitionTestUtils.h"
@@ -25,13 +24,14 @@ TEST_CLASS_WITH_FLAGS(Remesh, "MeshPartition.Modifier", EAutomationTestFlags::Ed
 	TUniquePtr<FScopedEditorWorld> ScopedEditorWorld;
 
 	UWorld* World = nullptr;
-	AMeshPartition* MegaMesh = nullptr;
+	AMeshPartition* MeshPartition = nullptr;
 	UMeshPartitionEditorComponent* MeshPartitionComponent = nullptr;
 	AActor* RemeshModifierActor = nullptr;
 	URemeshModifier* RemeshModifier = nullptr;
 
 	// CVars
-	TSharedPtr<FScopedTestEnvironment> TestEnvironment;
+	TOptional<TestUtils::FScopedCVarOverride> DisableDDCWrite;
+	TOptional<TestUtils::FScopedCVarOverride> DisableDDCRead;
 
 	// Test configuration
 	const FString LevelPath = TEXT("/MeshPartition/UnitTests/MeshPartitionUnitTests");
@@ -51,6 +51,7 @@ TEST_CLASS_WITH_FLAGS(Remesh, "MeshPartition.Modifier", EAutomationTestFlags::Ed
 		InitializationValues.CreateAISystem(false);
 		InitializationValues.AllowAudioPlayback(false);
 		InitializationValues.CreatePhysicsScene(true);
+		InitializationValues.CreateWorldPartition(true);
 
 
 #if AUTOMATED_TESTING_EXECUTE_IN_EDITOR_WORLD
@@ -65,16 +66,25 @@ TEST_CLASS_WITH_FLAGS(Remesh, "MeshPartition.Modifier", EAutomationTestFlags::Ed
 		ASSERT_THAT(IsNotNull(World, TEXT("World is not valid")));
 
 		// Need to disable DDC read/write, so that generated preview sections aren't written to DDC
-		TestEnvironment = FScopedTestEnvironment::Get();
-		TestEnvironment->SetConsoleVariableValue(TEXT("MegaMesh.Preview.EnableDDCWrite"), TEXT("0"));
-		TestEnvironment->SetConsoleVariableValue(TEXT("MegaMesh.Preview.EnableDDCRead"), TEXT("0"));
+		DisableDDCWrite.Emplace(TEXT("MegaMesh.Preview.EnableDDCWrite"), TEXT("0"));
+		DisableDDCRead.Emplace(TEXT("MegaMesh.Preview.EnableDDCRead"), TEXT("0"));
 
-		// Setup MegaMesh
-		MegaMesh = ImportHeightmap();
-		ASSERT_THAT(IsNotNull(MegaMesh, TEXT("Failed to import heightmap into mesh terrain.")));
+		// Setup MeshPartition
+		MeshPartition = ImportHeightmap();
+		ASSERT_THAT(IsNotNull(MeshPartition, TEXT("Failed to import heightmap into mesh terrain.")));
 
-		MeshPartitionComponent = Cast<UMeshPartitionEditorComponent>(MegaMesh->GetMeshPartitionComponent());
-		ASSERT_THAT(IsNotNull(MeshPartitionComponent, TEXT("Failed to get MegaMeshEditorComponent")));
+		MeshPartitionComponent = Cast<UMeshPartitionEditorComponent>(MeshPartition->GetMeshPartitionComponent());
+		ASSERT_THAT(IsNotNull(MeshPartitionComponent, TEXT("Failed to get MeshPartitionEditorComponent")));
+
+		// Setting default MeshPartitionDefinition
+		const FString TestAutomationDefaultDefinitionPath = TEXT("/MeshPartition/UnitTests/DataAssets/MPD_AutomationTestDefault.MPD_AutomationTestDefault");
+		UMeshPartitionDefinition* MeshPartitionDefinition = LoadObject<UMeshPartitionDefinition>(nullptr, TestAutomationDefaultDefinitionPath);
+		if (!MeshPartitionDefinition)
+		{
+			UE_LOG(LogMegaMeshEditor, Warning,
+				TEXT("Failed to load MeshPartitionDefinition (%s)"), *TestAutomationDefaultDefinitionPath);
+		}
+		MeshPartition->SetMeshPartitionDefinition(MeshPartitionDefinition);
 
 		// Pausing Transformer pipeline during the test to avoid potential race condition
 		MeshPartitionComponent->SetPauseTransformerPipeline(true);
@@ -90,20 +100,23 @@ TEST_CLASS_WITH_FLAGS(Remesh, "MeshPartition.Modifier", EAutomationTestFlags::Ed
 		// Create Remesh modifier and assign it to RemeshModifierActor as component
 		RemeshModifier = TestUtils::CreateRemeshModifier(
 			RemeshModifierActor,
-			MegaMesh, 
+			MeshPartition, 
 			FVector::ZeroVector,
 			UnscaledCoverage
 		);
 		RemeshModifier->SetTargetEdgeLength(10);
+
+		// Load WorldPartition Region
+		const bool bSuccessfullyLoadedWorldPartition = TestUtils::LoadWorldPartitionRegion(World);
+		ASSERT_THAT(IsTrue(bSuccessfullyLoadedWorldPartition, TEXT("Failed to load WorldPartition region")));
 	}
 
 	TEST_METHOD(Verify_RemeshModifier_Can_Be_Added)
 	{
 		TestCommandBuilder
-			.Do(TEXT("Notify MegaMesh that modifier has been assigned"), [this]()
+			.Do(TEXT("Notify MeshPartition that modifier has been assigned"), [this]()
 				{
 					MeshPartitionComponent->OnModifierAssigned();
-					MeshPartitionComponent->BuildMegaMeshPreviewSections(RemeshModifier->ComputeBounds());
 				})
 			.Until(TEXT("Wait for preview build completion"), [this]()
 				{
@@ -119,7 +132,7 @@ TEST_CLASS_WITH_FLAGS(Remesh, "MeshPartition.Modifier", EAutomationTestFlags::Ed
 	TEST_METHOD(Verify_RemeshModifier_Translation_Functions_Correctly)
 	{
 		TestCommandBuilder
-			.Do(TEXT("Notify MegaMesh that modifier has been assigned"), [this]()
+			.Do(TEXT("Notify MeshPartition that modifier has been assigned"), [this]()
 				{
 					MeshPartitionComponent->OnModifierAssigned();
 				})
@@ -147,7 +160,7 @@ TEST_CLASS_WITH_FLAGS(Remesh, "MeshPartition.Modifier", EAutomationTestFlags::Ed
 	TEST_METHOD(Verify_RemeshModifier_Rotation_Functions_Correctly)
 	{
 		TestCommandBuilder
-			.Do(TEXT("Notify MegaMesh that modifier has been assigned"), [this]()
+			.Do(TEXT("Notify MeshPartition that modifier has been assigned"), [this]()
 				{
 					MeshPartitionComponent->OnModifierAssigned();
 				})
@@ -175,7 +188,7 @@ TEST_CLASS_WITH_FLAGS(Remesh, "MeshPartition.Modifier", EAutomationTestFlags::Ed
 	TEST_METHOD(Verify_RemeshModifier_Scaling_Functions_Correctly)
 	{
 		TestCommandBuilder
-			.Do(TEXT("Notify MegaMesh that modifier has been assigned"), [this]()
+			.Do(TEXT("Notify MeshPartition that modifier has been assigned"), [this]()
 				{
 					MeshPartitionComponent->OnModifierAssigned();
 				})
@@ -203,7 +216,7 @@ TEST_CLASS_WITH_FLAGS(Remesh, "MeshPartition.Modifier", EAutomationTestFlags::Ed
 	TEST_METHOD(Verify_RemeshModifier_EdgeLength_Affects_Triangulation)
 	{
 		TestCommandBuilder
-			.Do(TEXT("Notify MegaMesh that modifier has been assigned"), [this]()
+			.Do(TEXT("Notify MeshPartition that modifier has been assigned"), [this]()
 				{
 					MeshPartitionComponent->OnModifierAssigned();
 				})
@@ -230,7 +243,7 @@ TEST_CLASS_WITH_FLAGS(Remesh, "MeshPartition.Modifier", EAutomationTestFlags::Ed
 	TEST_METHOD(Verify_RemeshModifier_Iterations_Are_Adjustable)
 	{
 		TestCommandBuilder
-			.Do(TEXT("Notify MegaMesh that modifier has been assigned"), [this]()
+			.Do(TEXT("Notify MeshPartition that modifier has been assigned"), [this]()
 				{
 					MeshPartitionComponent->OnModifierAssigned();
 				})
@@ -257,7 +270,7 @@ TEST_CLASS_WITH_FLAGS(Remesh, "MeshPartition.Modifier", EAutomationTestFlags::Ed
 	TEST_METHOD(Verify_RemeshModifier_Smoothing_Is_Adjustable)
 	{
 		TestCommandBuilder
-			.Do(TEXT("Notify MegaMesh that modifier has been assigned"), [this]()
+			.Do(TEXT("Notify MeshPartition that modifier has been assigned"), [this]()
 				{
 					MeshPartitionComponent->OnModifierAssigned();
 				})
@@ -284,7 +297,7 @@ TEST_CLASS_WITH_FLAGS(Remesh, "MeshPartition.Modifier", EAutomationTestFlags::Ed
 	TEST_METHOD(Verify_RemeshModifier_Supports_Component_Tags)
 	{
 		TestCommandBuilder
-			.Do(TEXT("Notify MegaMesh that modifier has been assigned"), [this]()
+			.Do(TEXT("Notify MeshPartition that modifier has been assigned"), [this]()
 				{
 					MeshPartitionComponent->OnModifierAssigned();
 				})
@@ -309,7 +322,7 @@ TEST_CLASS_WITH_FLAGS(Remesh, "MeshPartition.Modifier", EAutomationTestFlags::Ed
 	TEST_METHOD(Verify_RemeshModifier_Can_Be_Disabled_And_Enabled)
 	{
 		TestCommandBuilder
-			.Do(TEXT("Notify MegaMesh that modifier has been assigned"), [this]()
+			.Do(TEXT("Notify MeshPartition that modifier has been assigned"), [this]()
 				{
 					MeshPartitionComponent->OnModifierAssigned();
 				})
@@ -368,9 +381,9 @@ TEST_CLASS_WITH_FLAGS(Remesh, "MeshPartition.Modifier", EAutomationTestFlags::Ed
 
 	private:
 	/**
-	 * Imports a heightmap and creates a MegaMesh from it
+	 * Imports a heightmap and creates a MeshPartition from it
 	 * 
-	 * @return	Pointer to the created AMegaMesh, or nullptr if import failed
+	 * @return	Pointer to the created AMeshPartition, or nullptr if import failed
 	 */
 	AMeshPartition* ImportHeightmap() const
 	{
@@ -382,7 +395,7 @@ TEST_CLASS_WITH_FLAGS(Remesh, "MeshPartition.Modifier", EAutomationTestFlags::Ed
 		const TSharedPtr<IPlugin> Plugin = IPluginManager::Get().FindPlugin(TEXT("MeshPartition"));
 		if (!Plugin.IsValid())
 		{
-			UE_LOG(LogMegaMeshEditor, Error, TEXT("ImportHeightmap: MegaMesh plugin not found"));
+			UE_LOG(LogMegaMeshEditor, Error, TEXT("ImportHeightmap: MeshPartition plugin not found"));
 			return nullptr;
 		}
 
diff --git a/Engine/Plugins/Experimental/MeshPartition/Source/MeshPartitionEditor/Private/Test/MeshPartitionSplineModifierTests.cpp b/Engine/Plugins/Experimental/MeshPartition/Source/MeshPartitionEditor/Private/Test/MeshPartitionSplineModifierTests.cpp
--- a/Engine/Plugins/Experimental/MeshPartition/Source/MeshPartitionEditor/Private/Test/MeshPartitionSplineModifierTests.cpp
+++ b/Engine/Plugins/Experimental/MeshPartition/Source/MeshPartitionEditor/Private/Test/MeshPartitionSplineModifierTests.cpp
@@ -6,7 +6,6 @@
 #include "CQTest.h"
 #include "Editor.h"
 #include "MeshPartition.h"
-#include "MeshPartitionPreviewSection.h"
 #include "Components/SplineComponent.h"
 #include "Modifiers/MeshPartitionSplineModifier.h"
 #include "Test/MeshPartitionTestUtils.h"
@@ -14,7 +13,6 @@
 #include "Tests/AutomationEditorCommon.h"
 #include "ScopedTransaction.h"
 
-
 using namespace UE::MeshPartition;
 
 /**
@@ -25,7 +23,8 @@ TEST_CLASS_WITH_FLAGS(SplineModifierTests, "MeshPartition.Modifier.Splines", EAu
 	// Test environment and core objects
 	TUniquePtr<FScopedEditorWorld> ScopedEditorWorld;
 	// CVars
-	TSharedPtr<FScopedTestEnvironment> TestEnvironment;
+	TOptional<TestUtils::FScopedCVarOverride> DisableDDCWrite;
+	TOptional<TestUtils::FScopedCVarOverride> DisableDDCRead;
 
 	UWorld* World = nullptr;
 	AActor* SplineActor = nullptr;
@@ -103,9 +102,8 @@ TEST_CLASS_WITH_FLAGS(SplineModifierTests, "MeshPartition.Modifier.Splines", EAu
 		ASSERT_THAT(IsNotNull(World, TEXT("Failed to create valid test world")));
 
 		// Need to disable DDC read/write, so that generated preview sections aren't written to DDC
-		TestEnvironment = FScopedTestEnvironment::Get();
-		TestEnvironment->SetConsoleVariableValue(TEXT("MegaMesh.Preview.EnableDDCWrite"), TEXT("0"));
-		TestEnvironment->SetConsoleVariableValue(TEXT("MegaMesh.Preview.EnableDDCRead"), TEXT("0"));					
+		DisableDDCWrite.Emplace(TEXT("MegaMesh.Preview.EnableDDCWrite"), TEXT("0"));
+		DisableDDCRead.Emplace(TEXT("MegaMesh.Preview.EnableDDCRead"), TEXT("0"));
 		
 		// Create MeshPartition
 		MeshPartition = TestUtils::CreateTestMesh(World, MeshWidth, MeshHeight, MeshWidthVertexCount, MeshHeightVertexCount);
@@ -113,6 +111,13 @@ TEST_CLASS_WITH_FLAGS(SplineModifierTests, "MeshPartition.Modifier.Splines", EAu
 
 		MeshPartitionEditorComponent = Cast<UMeshPartitionEditorComponent>(MeshPartition->GetMeshPartitionComponent());
 		ASSERT_THAT(IsNotNull(MeshPartitionEditorComponent, TEXT("Failed to get MeshPartitionEditorComponent")));
+		
+		// Pausing the transformer pipeline during the test to avoid potential race condition
+		MeshPartitionEditorComponent->SetPauseTransformerPipeline(true);
+
+		// Load WorldPartition Region
+		const bool bSuccessfullyLoadedWorldPartition = TestUtils::LoadWorldPartitionRegion(World);
+		ASSERT_THAT(IsTrue(bSuccessfullyLoadedWorldPartition, TEXT("Failed to load WorldPartition region")));
 
 		// Create spline actor and component
 		SplineActor = World->SpawnActor<AActor>(AActor::StaticClass(), FTransform::Identity);
diff --git a/Engine/Plugins/Experimental/MeshPartition/Source/MeshPartitionEditor/Private/Test/MeshPartitionTestUtils.cpp b/Engine/Plugins/Experimental/MeshPartition/Source/MeshPartitionEditor/Private/Test/MeshPartitionTestUtils.cpp
--- a/Engine/Plugins/Experimental/MeshPartition/Source/MeshPartitionEditor/Private/Test/MeshPartitionTestUtils.cpp
+++ b/Engine/Plugins/Experimental/MeshPartition/Source/MeshPartitionEditor/Private/Test/MeshPartitionTestUtils.cpp
@@ -16,6 +16,9 @@
 #include "Modifiers/MeshPartitionPatchModifier.h"
 #include "Modifiers/MeshPartitionRemeshModifier.h"
 #include "Subsystems/EditorAssetSubsystem.h"
+#include "WorldPartition/WorldPartition.h"
+#include "WorldPartition/WorldPartitionEditorLoaderAdapter.h"
+#include "WorldPartition/LoaderAdapter/LoaderAdapterShape.h"
 
 namespace UE::MeshPartition::TestUtils
 {
@@ -36,22 +39,31 @@ namespace UE::MeshPartition::TestUtils
 
 	AMeshPartition* CreateTestMesh(UWorld* InWorld, int InWidth, int InHeight, int InWidthVertexCount, int InHeightVertexCount)
 	{
-		AMeshPartition* MegaMesh = InWorld->SpawnActor<AMeshPartition>(AMeshPartition::StaticClass(), FTransform::Identity);
-		UMeshPartitionEditorComponent* MegaMeshEditorComponent = NewObject<UMeshPartitionEditorComponent>(MegaMesh, UMeshPartitionEditorComponent::StaticClass(), TEXT("MegaMeshEditorComponent"));
-		MegaMeshEditorComponent->SetForceSynchronousPreviewSectionBuild(true);
-		MegaMesh->SetMeshPartitionComponent(MegaMeshEditorComponent);
+		AMeshPartition* MeshPartition = InWorld->SpawnActor<AMeshPartition>(AMeshPartition::StaticClass(), FTransform::Identity);
+		UMeshPartitionEditorComponent* MeshPartitionEditorComponent = NewObject<UMeshPartitionEditorComponent>(MeshPartition, UMeshPartitionEditorComponent::StaticClass(), TEXT("MegaMeshEditorComponent"));
+		MeshPartitionEditorComponent->SetForceSynchronousPreviewSectionBuild(true);
+		MeshPartition->SetMeshPartitionComponent(MeshPartitionEditorComponent);
+
+		const FString TestAutomationDefaultDefinitionPath = TEXT("/MeshPartition/UnitTests/DataAssets/MPD_AutomationTestDefault.MPD_AutomationTestDefault");
+		UMeshPartitionDefinition* MeshPartitionDefinition = LoadObject<UMeshPartitionDefinition>(nullptr, TestAutomationDefaultDefinitionPath);
+		if (!MeshPartitionDefinition)
+		{
+			UE_LOG(LogMegaMeshEditor, Warning,
+				TEXT("CreateTestMesh: Failed to load MeshPartitionDefinition (%s)"), *TestAutomationDefaultDefinitionPath);
+		}
+		MeshPartition->SetMeshPartitionDefinition(MeshPartitionDefinition);
 
 		Geometry::FRectangleMeshGenerator RectGen;
 		RectGen.Width = InWidth;
 		RectGen.Height = InHeight;
 		RectGen.WidthVertexCount = InWidthVertexCount;
 		RectGen.HeightVertexCount = InHeightVertexCount;
-		MegaMeshEditorComponent->SpawnBaseModifier(FDynamicMesh3(&RectGen.Generate()), {}, FTransform::Identity);
+		MeshPartitionEditorComponent->SpawnBaseModifier(FDynamicMesh3(&RectGen.Generate()), {}, FTransform::Identity);
 
-		return MegaMesh;
+		return MeshPartition;
 	}
 
-	AActor* CreateRectangleModifier(UMeshPartitionEditorComponent* MegaMeshEditorComponent, const double Width, const double Height, const int32 WidthVertexCount,
+	AActor* CreateRectangleModifier(UMeshPartitionEditorComponent* MeshPartitionEditorComponent, const double Width, const double Height, const int32 WidthVertexCount,
 		const int32 HeightVertexCount, TArray<FBox>& OutBounds, const FTransform& Transform)
 	{
 		// Configure rectangle mesh generator
@@ -62,7 +74,7 @@ namespace UE::MeshPartition::TestUtils
 		RectGen.HeightVertexCount = HeightVertexCount;
 
 		// Spawn a base modifier actor using the generated mesh
-		AActor* BaseModifier = MegaMeshEditorComponent->SpawnBaseModifier(
+		AActor* BaseModifier = MeshPartitionEditorComponent->SpawnBaseModifier(
 			FDynamicMesh3(&RectGen.Generate()),
 			{},
 			Transform
@@ -76,7 +88,7 @@ namespace UE::MeshPartition::TestUtils
 		return BaseModifier;
 	}
 
-	UPatchModifier* CreatePatchModifier(AActor* PatchActor, AMeshPartition* MegaMesh, const FVector& ModifierWorldLocation, const float Priority)
+	UPatchModifier* CreatePatchModifier(AActor* PatchActor, AMeshPartition* MeshPartition, const FVector& ModifierWorldLocation, const float Priority)
 	{
 		if (!IsValid(PatchActor))
 		{
@@ -84,9 +96,9 @@ namespace UE::MeshPartition::TestUtils
 			return nullptr;
 		}
 
-		if (!IsValid(MegaMesh))
+		if (!IsValid(MeshPartition))
 		{
-			UE_LOG(LogMegaMeshEditor, Error, TEXT("CreatePatchModifier: MegaMesh is null or invalid"));
+			UE_LOG(LogMegaMeshEditor, Error, TEXT("CreatePatchModifier: MeshPartition is null or invalid"));
 			return nullptr;
 		}
 
@@ -111,7 +123,7 @@ namespace UE::MeshPartition::TestUtils
 		PatchModifier->SetWorldLocation(ModifierWorldLocation);
 		PatchModifier->SetPriority(Priority);
 		PatchActor->AddInstanceComponent(PatchModifier);
-		PatchModifier->SetAffectedMeshPartition(MegaMesh);
+		PatchModifier->SetAffectedMeshPartition(MeshPartition);
 		// Attach to root component
 		if (USceneComponent* RootComp = PatchActor->GetRootComponent())
 		{
@@ -128,7 +140,7 @@ namespace UE::MeshPartition::TestUtils
 		return PatchModifier;
 	}
 
-	URemeshModifier* CreateRemeshModifier(AActor* Actor, AMeshPartition* MegaMesh, const FVector& ModifierWorldLocation,
+	URemeshModifier* CreateRemeshModifier(AActor* Actor, AMeshPartition* MeshPartition, const FVector& ModifierWorldLocation,
 		const FVector3d& UnscaledCoverage)
 	{
 		if (!IsValid(Actor))
@@ -137,9 +149,9 @@ namespace UE::MeshPartition::TestUtils
 			return nullptr;
 		}
 
-		if (!IsValid(MegaMesh))
+		if (!IsValid(MeshPartition))
 		{
-			UE_LOG(LogMegaMeshEditor, Error, TEXT("CreateRemeshModifier: MegaMesh is null or invalid"));
+			UE_LOG(LogMegaMeshEditor, Error, TEXT("CreateRemeshModifier: MeshPartition is null or invalid"));
 			return nullptr;
 		}
 
@@ -160,7 +172,7 @@ namespace UE::MeshPartition::TestUtils
 		RemeshModifier->SetWorldLocation(ModifierWorldLocation);
 		RemeshModifier->SetUnscaledCoverage(UnscaledCoverage);
 		Actor->AddInstanceComponent(RemeshModifier);
-		RemeshModifier->SetAffectedMeshPartition(MegaMesh);
+		RemeshModifier->SetAffectedMeshPartition(MeshPartition);
 		// Attach to root component
 		if (USceneComponent* RootComp = Actor->GetRootComponent())
 		{
@@ -180,10 +192,10 @@ namespace UE::MeshPartition::TestUtils
 
 	AMeshPartition* CreateTestMeshSectioned(UWorld* InWorld, int InSectionCountX, int InSectionCountY, int InSectionWidth, int InSectionHeight, int InSectionWidthVertexCount, int InSectionHeightVertexCount)
 	{
-		AMeshPartition* MegaMesh = InWorld->SpawnActor<AMeshPartition>(AMeshPartition::StaticClass(), FTransform::Identity);
-		UMeshPartitionEditorComponent* MegaMeshEditorComponent = NewObject<UMeshPartitionEditorComponent>(MegaMesh, UMeshPartitionEditorComponent::StaticClass(), TEXT("MegaMeshEditorComponent"));
-		MegaMeshEditorComponent->SetForceSynchronousPreviewSectionBuild(true);
-		MegaMesh->SetMeshPartitionComponent(MegaMeshEditorComponent);
+		AMeshPartition* MeshPartition = InWorld->SpawnActor<AMeshPartition>(AMeshPartition::StaticClass(), FTransform::Identity);
+		UMeshPartitionEditorComponent* MeshPartitionEditorComponent = NewObject<UMeshPartitionEditorComponent>(MeshPartition, UMeshPartitionEditorComponent::StaticClass(), TEXT("MegaMeshEditorComponent"));
+		MeshPartitionEditorComponent->SetForceSynchronousPreviewSectionBuild(true);
+		MeshPartition->SetMeshPartitionComponent(MeshPartitionEditorComponent);
 
 		for (int SectionIndexY = 0; SectionIndexY < InSectionCountY; ++SectionIndexY)
 		{
@@ -194,11 +206,11 @@ namespace UE::MeshPartition::TestUtils
 				RectGen.Height = InSectionHeight;
 				RectGen.WidthVertexCount = InSectionWidthVertexCount;
 				RectGen.HeightVertexCount = InSectionHeightVertexCount;
-				MegaMeshEditorComponent->SpawnBaseModifier(FDynamicMesh3(&RectGen.Generate()), {}, FTransform(FVector(SectionIndexX * InSectionWidth, SectionIndexY * InSectionHeight, 0.)));
+				MeshPartitionEditorComponent->SpawnBaseModifier(FDynamicMesh3(&RectGen.Generate()), {}, FTransform(FVector(SectionIndexX * InSectionWidth, SectionIndexY * InSectionHeight, 0.)));
 			}
 		}
 
-		return MegaMesh;
+		return MeshPartition;
 	}
 	
 	bool LoadMeshes(TArray<TObjectPtr<UStaticMesh>>& OutMeshes, TArrayView<const TSoftObjectPtr<UStaticMesh>> InSoftPointers)
@@ -219,6 +231,29 @@ namespace UE::MeshPartition::TestUtils
 		return true;
 	}
 
+	bool LoadWorldPartitionRegion(UWorld* InWorld, const FVector& InWorldBoxExtent)
+	{
+		FBox WorldBox(-InWorldBoxExtent, InWorldBoxExtent);
+		UWorldPartition* WorldPartition = InWorld->GetWorldPartition();
+		if (!WorldPartition)
+		{
+			UE_LOG(LogMegaMeshEditor, Error, TEXT("WorldPartition is null."));
+			return false;
+		}
+
+		UWorldPartitionEditorLoaderAdapter* EditorLoaderAdapter =
+			WorldPartition->CreateEditorLoaderAdapter<FLoaderAdapterShape>(InWorld, WorldBox, TEXT("Loaded Region"));
+		if (!EditorLoaderAdapter)
+		{
+			UE_LOG(LogMegaMeshEditor, Error, TEXT("Editor loader adapter creation failed."));
+			return false;
+		}
+		EditorLoaderAdapter->GetLoaderAdapter()->SetUserCreated(true);
+		EditorLoaderAdapter->GetLoaderAdapter()->Load();
+
+		return true;
+	}
+
 	bool ValidatePreviewMeshAgainstGroundTruth(const UMeshPartitionEditorComponent* InMeshPartitionComponent,
 		const TArray<FString>& InGroundTruthPaths, FString& OutErrorMessage)
 	{
diff --git a/Engine/Plugins/Experimental/MeshPartition/Source/MeshPartitionEditor/Private/Test/MeshPartitionTexturePatchModifierTests.cpp b/Engine/Plugins/Experimental/MeshPartition/Source/MeshPartitionEditor/Private/Test/MeshPartitionTexturePatchModifierTests.cpp
--- a/Engine/Plugins/Experimental/MeshPartition/Source/MeshPartitionEditor/Private/Test/MeshPartitionTexturePatchModifierTests.cpp
+++ b/Engine/Plugins/Experimental/MeshPartition/Source/MeshPartitionEditor/Private/Test/MeshPartitionTexturePatchModifierTests.cpp
@@ -29,7 +29,8 @@ TEST_CLASS_WITH_FLAGS(TexturePatchGoldenPath, "MeshPartition.Modifier", EAutomat
 	// Test environment and core objects
 	TUniquePtr<FScopedEditorWorld> ScopedEditorWorld;
 	// CVars
-	TSharedPtr<FScopedTestEnvironment> TestEnvironment;
+	TOptional<TestUtils::FScopedCVarOverride> DisableDDCWrite;
+	TOptional<TestUtils::FScopedCVarOverride> DisableDDCRead;
 
 	UWorld* World = nullptr;
 	AMeshPartition* MeshPartition = nullptr;
@@ -72,9 +73,8 @@ TEST_CLASS_WITH_FLAGS(TexturePatchGoldenPath, "MeshPartition.Modifier", EAutomat
 		ASSERT_THAT(IsNotNull(World, TEXT("Failed to create valid test world")));
 
 		// Need to disable DDC read/write, so that generated preview sections aren't written to DDC
-		TestEnvironment = FScopedTestEnvironment::Get();
-		TestEnvironment->SetConsoleVariableValue(TEXT("MegaMesh.Preview.EnableDDCWrite"), TEXT("0"));
-		TestEnvironment->SetConsoleVariableValue(TEXT("MegaMesh.Preview.EnableDDCRead"), TEXT("0"));
+		DisableDDCWrite.Emplace(TEXT("MegaMesh.Preview.EnableDDCWrite"), TEXT("0"));
+		DisableDDCRead.Emplace(TEXT("MegaMesh.Preview.EnableDDCRead"), TEXT("0"));
 
 		MeshPartition = TestUtils::CreateTestMesh(
 			World, MeshWidth, MeshHeight, MeshWidthVertexCount, MeshHeightVertexCount
@@ -102,6 +102,10 @@ TEST_CLASS_WITH_FLAGS(TexturePatchGoldenPath, "MeshPartition.Modifier", EAutomat
 		TexturePatchHeightEntry->SetTextureAsset(Texture);
 
 		TexturePatchModifier->SetHeightChannel(TexturePatchHeightEntry);
+
+		// Load WorldPartition Region
+		const bool bSuccessfullyLoadedWorldPartition = TestUtils::LoadWorldPartitionRegion(World);
+		ASSERT_THAT(IsTrue(bSuccessfullyLoadedWorldPartition, TEXT("Failed to load WorldPartition region")));
 	}
 
 	TEST_METHOD(Verify_TexturePatch_Modifier_TextureAsset_Input_Projects_To_MeshTerrain)
diff --git a/Engine/Plugins/Experimental/MeshPartition/Source/MeshPartitionEditor/Public/Test/MeshPartitionTestUtils.h b/Engine/Plugins/Experimental/MeshPartition/Source/MeshPartitionEditor/Public/Test/MeshPartitionTestUtils.h
--- a/Engine/Plugins/Experimental/MeshPartition/Source/MeshPartitionEditor/Public/Test/MeshPartitionTestUtils.h
+++ b/Engine/Plugins/Experimental/MeshPartition/Source/MeshPartitionEditor/Public/Test/MeshPartitionTestUtils.h
@@ -7,11 +7,9 @@
 #include "MeshPartitionEditorComponent.h"
 #include "MeshPartitionEditorModule.h"
 
-#include "Containers/ArrayView.h"
 #include "Containers/Array.h"
 #include "Misc/AutomationTest.h"
 #include "UObject/SoftObjectPtr.h"
-#include "MeshPartitionMeshData.h"
 #include "Engine/StaticMesh.h"
 #include "MeshDescription.h"
 #include "Modifiers/MeshPartitionPatchModifier.h"
@@ -43,39 +41,39 @@ MESHPARTITIONEDITOR_API AMeshPartition* CreateTestMesh(UWorld* InWorld, int InWi
 /**
 * Creates a rectangular mesh modifier with specified complexity and transform
 *
-* @param MegaMeshEditorComponent	Width of rectangle
-* @param Width						Width of rectangle
-* @param Height					Height of rectangle
-* @param WidthVertexCount			Number of vertices along width of the rectangle
-* @param HeightVertexCount			Number of vertices along height of the rectangle
-* @param OutBounds					Array to store the resulting bounding box
-* @param Transform					World transform to apply to the modifier (defaults to identity)
-* @return							Pointer to the created actor
+* @param MeshPartitionEditorComponent	Component for affected MeshPartition
+* @param Width							Width of rectangle
+* @param Height							Height of rectangle
+* @param WidthVertexCount				Number of vertices along width of the rectangle
+* @param HeightVertexCount				Number of vertices along height of the rectangle
+* @param OutBounds						Array to store the resulting bounding box
+* @param Transform						World transform to apply to the modifier (defaults to identity)
+* @return								Pointer to the created actor
 */
-AActor* CreateRectangleModifier(UMeshPartitionEditorComponent* MegaMeshEditorComponent, const double Width, const double Height, const int32 WidthVertexCount,
+AActor* CreateRectangleModifier(UMeshPartitionEditorComponent* MeshPartitionEditorComponent, const double Width, const double Height, const int32 WidthVertexCount,
 	const int32 HeightVertexCount, TArray<FBox>& OutBounds, const FTransform& Transform = FTransform::Identity);
 
 /**
  * Creates Patch Modifier with specified world location and build priority
  *
  * @param PatchActor				Modifier Actor
- * @param MegaMesh					Affected MegaMesh by modifier
+ * @param MeshPartition				Affected MeshPartition by modifier
  * @param ModifierWorldLocation		Location of the modifier actor
  * @param Priority					Build priority (Default = 0.0)
  * @return							Pointer to the created UPatchModifier
  */
-UPatchModifier* CreatePatchModifier(AActor* PatchActor, AMeshPartition* MegaMesh, const FVector& ModifierWorldLocation, const float Priority = 0.0);
+UPatchModifier* CreatePatchModifier(AActor* PatchActor, AMeshPartition* MeshPartition, const FVector& ModifierWorldLocation, const float Priority = 0.0);
 
 /**
  * Creates Remesh modifier with specified Location and UnscaledCoverage property
  *
  * @param Actor						Modifier's actor
- * @param MegaMesh					Affected MegaMesh
+ * @param MeshPartition				Affected MeshPartition
  * @param ModifierWorldLocation		Location of the modifier
  * @param UnscaledCoverage			Remesh Modifier's Unscaled Coverage property (Bounding box)
- * @return							Returns MegaMeshRemeshModifier
+ * @return							Returns RemeshModifier
  */
-MESHPARTITIONEDITOR_API URemeshModifier* CreateRemeshModifier(AActor* Actor, AMeshPartition* MegaMesh, const FVector& ModifierWorldLocation, const FVector3d& UnscaledCoverage);
+MESHPARTITIONEDITOR_API URemeshModifier* CreateRemeshModifier(AActor* Actor, AMeshPartition* MeshPartition, const FVector& ModifierWorldLocation, const FVector3d& UnscaledCoverage);
 
 /**
 * Helper utility that creates a multiple base modifier plane meshes with specified worldspace width/height and vertex counts.
@@ -87,6 +85,15 @@ MESHPARTITIONEDITOR_API AMeshPartition* CreateTestMeshSectioned(UWorld* InWorld,
 */
 MESHPARTITIONEDITOR_API bool LoadMeshes(TArray<TObjectPtr<UStaticMesh>>& OutMeshes, TArrayView<const TSoftObjectPtr<UStaticMesh>> InSoftPointers);
 
+/**
+ * Loads a region of partitioned world.
+ *
+ * @param InWorld World pointer.
+ * @param InWorldBoxExtent FVector extent for WorldBox that is to be loaded, defaults to 10k if unspecified.
+ * @return true if the region was loaded successfully.
+ */
+MESHPARTITIONEDITOR_API bool LoadWorldPartitionRegion(UWorld* InWorld, const FVector& InWorldBoxExtent = FVector(10000));
+
 /**
 * Helper utility that compares the vertex positions of a saved ground truth static mesh asset with the generated dynamic mesh.
 */
@@ -185,7 +192,49 @@ MESHPARTITIONEDITOR_API bool CreateOrUpdateReference(const UMeshPartitionEditorC
                                                      const FString& ReferencePackage,
                                                      const TArray<FString>& ReferenceMeshNames,
                                                      const FString& TestName);
-} // namespace UE::MeshPartition::Utils
+
+/**
+ * RAII guard that saves a CVar's current value on construction,
+ * forces it to a new value for the duration of the scope, then restores it
+ * on destruction, even if the test throws or fails early.
+ */
+struct FScopedCVarOverride
+{
+	FScopedCVarOverride(const TCHAR* InName, const TCHAR* NewValue, const EConsoleVariableFlags SetBy = ECVF_SetByTemp)
+		: CVar(IConsoleManager::Get().FindConsoleVariable(InName))
+		, SetByPriority(SetBy)
+	{
+		if (CVar)
+		{
+			OriginalValue = CVar->GetString();
+			CVar->Set(NewValue, SetBy);
+		}
+		else
+		{
+			UE_LOG(LogMegaMeshEditor, Warning, TEXT("FScopedCVarOverride: CVar '%s' not found."), InName);
+		}
+	}
+
+	~FScopedCVarOverride()
+	{
+		if (CVar)
+		{
+			CVar->Set(*OriginalValue, SetByPriority);
+		}
+	}
+
+	// Non-copyable, non-movable, must live on the stack
+	FScopedCVarOverride(const FScopedCVarOverride&) = delete;
+	FScopedCVarOverride& operator=(const FScopedCVarOverride&) = delete;
+	FScopedCVarOverride(FScopedCVarOverride&&) = delete;
+	FScopedCVarOverride& operator=(FScopedCVarOverride&&) = delete;
+
+private:
+	IConsoleVariable* CVar;
+	FString OriginalValue;
+	EConsoleVariableFlags SetByPriority;
+};
+} // namespace UE::MeshPartition::TestUtils
 
 
 #endif // WITH_DEV_AUTOMATION_TESTS
\ No newline at end of file
diff --git a/Engine/Plugins/Experimental/MeshPartitionWater/Source/MeshPartitionWater/Private/Test/MeshPartitionWaterModifierTests.cpp b/Engine/Plugins/Experimental/MeshPartitionWater/Source/MeshPartitionWater/Private/Test/MeshPartitionWaterModifierTests.cpp
--- a/Engine/Plugins/Experimental/MeshPartitionWater/Source/MeshPartitionWater/Private/Test/MeshPartitionWaterModifierTests.cpp
+++ b/Engine/Plugins/Experimental/MeshPartitionWater/Source/MeshPartitionWater/Private/Test/MeshPartitionWaterModifierTests.cpp
@@ -23,19 +23,20 @@
 using namespace UE::MeshPartition;
 
 /**
- * Tests for MegaMesh Water Modifier
+ * Tests for MeshPartition Water Modifier
  */
 TEST_CLASS_WITH_FLAGS(Water, "MeshPartition", EAutomationTestFlags::EditorContext | EAutomationTestFlags::EngineFilter)
 {
 	// Test environment and core objects
 	TUniquePtr<FScopedEditorWorld> ScopedEditorWorld;
 	// CVars
-	TSharedPtr<FScopedTestEnvironment> TestEnvironment;
+	TOptional<TestUtils::FScopedCVarOverride> DisableDDCWrite;
+	TOptional<TestUtils::FScopedCVarOverride> DisableDDCRead;
 
 	// MegaMesh components under test
 	UWorld* World = nullptr;
 	AMeshPartition* MegaMesh = nullptr;
-	UMeshPartitionEditorComponent* MegaMeshEditorComponent = nullptr;
+	UMeshPartitionEditorComponent* MeshPartitionEditorComponent = nullptr;
 	URiverModifier* RiverModifier = nullptr;
 	ULakeModifier* LakeModifier = nullptr;
 
@@ -43,6 +44,7 @@ TEST_CLASS_WITH_FLAGS(Water, "MeshPartition", EAutomationTestFlags::EditorContex
 	TArray<TObjectPtr<UStaticMesh>> GroundTruth;
 
 	const FString LevelPath = TEXT("/MeshPartition/UnitTests/MeshPartitionUnitTests");
+	const FString ReferencePath = TEXT("/MeshPartitionWater/UnitTests");
 
 	/*
 	 * Defines the types of water body configurations to test
@@ -66,28 +68,25 @@ TEST_CLASS_WITH_FLAGS(Water, "MeshPartition", EAutomationTestFlags::EditorContex
 	};
 
 	/*
-	 * Returns ground truth mesh paths for each water test type.
+	 * Returns ground truth mesh asset names for each water test type.
 	 */
-	static const TMap<EWaterTestType, TSoftObjectPtr<UStaticMesh>>& GetGroundTruthPath()
+	static const TMap<EWaterTestType, const FString>& GetGroundTruthAsset()
 	{
-		static TMap<EWaterTestType, TSoftObjectPtr<UStaticMesh>> GroundTruthPaths;
+		static TMap<EWaterTestType, const FString> GroundTruthPaths;
 
 		if (GroundTruthPaths.Num() == 0)
 		{
 			GroundTruthPaths.Add(
 				EWaterTestType::Lake,
-				TSoftObjectPtr<UStaticMesh>(FSoftObjectPath(
-					TEXT("/MeshPartitionWater/UnitTests/Lake_PreviewSection.Lake_PreviewSection")))
+				TEXT("Lake_PreviewSection")
 			);
 			GroundTruthPaths.Add(
 				EWaterTestType::River,
-				TSoftObjectPtr<UStaticMesh>(FSoftObjectPath(
-					TEXT("/MeshPartitionWater/UnitTests/River_PreviewSection.River_PreviewSection")))
+				TEXT("River_PreviewSection")
 			);
 			GroundTruthPaths.Add(
 				EWaterTestType::LakeAndRiver,
-				TSoftObjectPtr<UStaticMesh>(FSoftObjectPath(
-					TEXT("/MeshPartitionWater/UnitTests/CombinedLakeAndRiverModifierGroundTruth.CombinedLakeAndRiverModifierGroundTruth")))
+				TEXT("CombinedLakeAndRiverModifierGroundTruth")
 			);
 		}
 
@@ -106,26 +105,24 @@ TEST_CLASS_WITH_FLAGS(Water, "MeshPartition", EAutomationTestFlags::EditorContex
 		InitializationValues.AllowAudioPlayback(false);
 		InitializationValues.CreatePhysicsScene(true);
 
+#if AUTOMATED_TESTING_EXECUTE_IN_EDITOR_WORLD
+		World = GEditor->GetEditorWorldContext().World();
+#else
 		// Spawn world
 		ScopedEditorWorld = FAutomationEditorCommonUtils::CreateScopedEditorWorld(LevelPath, InitializationValues);
 		ASSERT_THAT(IsNotNull(ScopedEditorWorld.Get(), TEXT("Failed to create scoped editor world")));
 
 		World = ScopedEditorWorld->GetWorld();
-		ASSERT_THAT(IsNotNull(World, TEXT("Failed to create valid test world")));
+#endif
+		ASSERT_THAT(IsNotNull(World, TEXT("World is not valid")));
 
 		// Need to disable DDC read/write, so that generated preview sections aren't written to DDC
-		TestEnvironment = FScopedTestEnvironment::Get();
-		TestEnvironment->SetConsoleVariableValue(TEXT("MegaMesh.Preview.EnableDDCWrite"), TEXT("0"));
-		TestEnvironment->SetConsoleVariableValue(TEXT("MegaMesh.Preview.EnableDDCRead"), TEXT("0"));
+		DisableDDCWrite.Emplace(TEXT("MegaMesh.Preview.EnableDDCWrite"), TEXT("0"));
+		DisableDDCRead.Emplace(TEXT("MegaMesh.Preview.EnableDDCRead"), TEXT("0"));
 	}
 
 	AFTER_EACH()
 	{
-		if (MegaMesh && MegaMesh->IsValidLowLevel())
-		{
-			MegaMesh->Destroy();
-		}
-
 		ScopedEditorWorld.Reset();
 		CollectGarbage(GARBAGE_COLLECTION_KEEPFLAGS);
 	}
@@ -148,20 +145,13 @@ TEST_CLASS_WITH_FLAGS(Water, "MeshPartition", EAutomationTestFlags::EditorContex
 	TEST_METHOD(Verify_RiverModifier_MaxZDistance_Updates_Mesh)
 	{
 		TestCommandBuilder
-			.Do(TEXT("Load ground truth meshes"), [this]()
+			.Do(TEXT("Create MeshTerrain"), [this]()
 				{
-					const bool bSuccessfullyLoadedMeshes = TestUtils::LoadMeshes(GroundTruth,
-						{
-							TSoftObjectPtr<UStaticMesh>(FSoftObjectPath(TEXT("/MeshPartitionWater/UnitTests/RiverModifierDefaultZHeightGroundTruth.RiverModifierDefaultZHeightGroundTruth"))),
-							TSoftObjectPtr<UStaticMesh>(FSoftObjectPath(TEXT("/MeshPartitionWater/UnitTests/RiverModifierModifiedZHeightGroundTruth.RiverModifierModifiedZHeightGroundTruth")))
-						});
-					ASSERT_THAT(IsTrue(bSuccessfullyLoadedMeshes, TEXT("Failed to load ground truth meshes")));
-
 					// Create sectioned MegaMesh for testing
 					MegaMesh = TestUtils::CreateTestMeshSectioned(World, 2, 2, 10000, 10000, 50, 50);
 					ASSERT_THAT(IsNotNull(MegaMesh, TEXT("Failed to create MegaMesh")));
 
-					MegaMeshEditorComponent = Cast<UMeshPartitionEditorComponent>(MegaMesh->GetMeshPartitionComponent());
+					MeshPartitionEditorComponent = Cast<UMeshPartitionEditorComponent>(MegaMesh->GetMeshPartitionComponent());
 				})
 			.Then(TEXT("Create river modifier"), [this]()
 				{
@@ -175,18 +165,18 @@ TEST_CLASS_WITH_FLAGS(Water, "MeshPartition", EAutomationTestFlags::EditorContex
 						EWaterTestType::River, DefaultRiverConfig
 					);
 
-					MegaMeshEditorComponent->OnModifierAssigned();
+					MeshPartitionEditorComponent->OnModifierAssigned();
 				})
 			.Until(TEXT("Wait for preview section build"), [this]()
 				{
-					return !MegaMeshEditorComponent->IsAnyPreviewSectionBuildActive();
+					return !MeshPartitionEditorComponent->IsAnyPreviewSectionBuildActive();
 				})
 			.Then(TEXT("Validate initial mesh"), [this]()
 				{
 					APreviewSection* PreviewSection = nullptr;
 					int32 PreviewSectionCount = 0;
 
-					MegaMeshEditorComponent->ForAllPreviewSections([&](APreviewSection* InPreviewSection)
+					MeshPartitionEditorComponent->ForAllPreviewSections([&](APreviewSection* InPreviewSection)
 						{
 							PreviewSection = InPreviewSection;
 							++PreviewSectionCount;
@@ -195,8 +185,8 @@ TEST_CLASS_WITH_FLAGS(Water, "MeshPartition", EAutomationTestFlags::EditorContex
 					ASSERT_THAT(AreEqual(PreviewSectionCount, 1, TEXT("PreviewSection count should be 1")));
 					ASSERT_THAT(IsNotNull(PreviewSection, TEXT("PreviewSection should not be null")));
 
-					const bool bAreMeshesEqual = TestUtils::CompareMesh(GroundTruth[0], PreviewSection->GetPreviewMesh().Get());
-					ASSERT_THAT(IsTrue(bAreMeshesEqual, TEXT("Meshes are not equal")));
+					const FString GroundTruthAsset = TEXT("RiverModifierDefaultZHeightGroundTruth");
+					CompareToReferenceAsset(GroundTruthAsset, MeshPartitionEditorComponent);
 
 					// Update MaxZDistance parameter
 					RiverModifier->SetMaxZDistance(1.f);
@@ -204,17 +194,17 @@ TEST_CLASS_WITH_FLAGS(Water, "MeshPartition", EAutomationTestFlags::EditorContex
 						FBox(FVector3d(-HALF_WORLD_MAX, -HALF_WORLD_MAX, -HALF_WORLD_MAX),
 						FVector3d(HALF_WORLD_MAX, HALF_WORLD_MAX, HALF_WORLD_MAX)) },
 						EChangeType::StateChange);
-					MegaMeshEditorComponent->Update();
+					MeshPartitionEditorComponent->Update();
 				})
 			.Until(TEXT("Wait for preview section build"), [this]()
 				{
-					return !MegaMeshEditorComponent->IsAnyPreviewSectionBuildActive();
+					return !MeshPartitionEditorComponent->IsAnyPreviewSectionBuildActive();
 				})
 			.Then(TEXT("Validate updated mesh with MaxZDistance"), [this]()
 				{
 					APreviewSection* PreviewSection = nullptr;
 					int32 PreviewSectionCount = 0;
-					MegaMeshEditorComponent->ForAllPreviewSections([&](APreviewSection* InPreviewSection)
+					MeshPartitionEditorComponent->ForAllPreviewSections([&](APreviewSection* InPreviewSection)
 						{
 							PreviewSection = InPreviewSection;
 							++PreviewSectionCount;
@@ -223,8 +213,8 @@ TEST_CLASS_WITH_FLAGS(Water, "MeshPartition", EAutomationTestFlags::EditorContex
 					ASSERT_THAT(AreEqual(PreviewSectionCount, 1, TEXT("PreviewSectionCount should be 1")));
 					ASSERT_THAT(IsNotNull(PreviewSection, TEXT("PreviewSection should not be null")));
 
-					const bool bAreMeshesEqual = TestUtils::CompareMesh(GroundTruth[1], PreviewSection->GetPreviewMesh().Get());
-					ASSERT_THAT(IsTrue(bAreMeshesEqual, TEXT("Meshes are not equal")));
+					const FString GroundTruthAsset = TEXT("RiverModifierModifiedZHeightGroundTruth");
+					CompareToReferenceAsset(GroundTruthAsset, MeshPartitionEditorComponent);
 				});
 	}
 
@@ -240,17 +230,10 @@ TEST_CLASS_WITH_FLAGS(Water, "MeshPartition", EAutomationTestFlags::EditorContex
 		TestCommandBuilder
 			.Do(TEXT("Setup test environment"), [this, Type]()
 				{
-					const auto& GroundTruthPath = GetGroundTruthPath()[Type];
-					const bool bSuccessfullyLoadedMeshes = TestUtils::LoadMeshes(
-						GroundTruth, { GroundTruthPath }
-					);
-
-					ASSERT_THAT(IsTrue(bSuccessfullyLoadedMeshes, TEXT("Failed to load the ground truth meshes")));
-
 					MegaMesh = TestUtils::CreateTestMesh(World, 20000, 20000, 100, 100);
 					ASSERT_THAT(IsNotNull(MegaMesh, TEXT("Failed to create test MegaMesh")));
 
-					MegaMeshEditorComponent = Cast<UMeshPartitionEditorComponent>(MegaMesh->GetMeshPartitionComponent());
+					MeshPartitionEditorComponent = Cast<UMeshPartitionEditorComponent>(MegaMesh->GetMeshPartitionComponent());
 				})
 			.Then(TEXT("Create water body modifiers"), [this, Type]()
 				{
@@ -292,17 +275,17 @@ TEST_CLASS_WITH_FLAGS(Water, "MeshPartition", EAutomationTestFlags::EditorContex
 					}
 
 					// Notify MegaMesh that modifiers have been assigned
-					MegaMeshEditorComponent->OnModifierAssigned();
+					MeshPartitionEditorComponent->OnModifierAssigned();
 				})
 			.Until(TEXT("Wait for preview section generation"), [this]()
 				{
-					return !MegaMeshEditorComponent->IsAnyPreviewSectionBuildActive();
+					return !MeshPartitionEditorComponent->IsAnyPreviewSectionBuildActive();
 				})
-			.Then(TEXT("Validate preview section and mesh"), [this]()
+			.Then(TEXT("Validate preview section and mesh"), [this, Type]()
 				{
 					APreviewSection* PreviewSection = nullptr;
 					int32 PreviewSectionCount = 0;
-					MegaMeshEditorComponent->ForAllPreviewSections([&](APreviewSection* InPreviewSection)
+					MeshPartitionEditorComponent->ForAllPreviewSections([&](APreviewSection* InPreviewSection)
 						{
 							PreviewSection = InPreviewSection;
 							++PreviewSectionCount;
@@ -312,8 +295,8 @@ TEST_CLASS_WITH_FLAGS(Water, "MeshPartition", EAutomationTestFlags::EditorContex
 					ASSERT_THAT(AreEqual(PreviewSectionCount, 1, TEXT("Expected PreviewSectionCount to be 1")));
 					ASSERT_THAT(IsNotNull(PreviewSection, TEXT("PreviewSection can not be null")));
 
-					const bool bAreMeshesEqual = TestUtils::CompareMesh(GroundTruth[0], PreviewSection->GetPreviewMesh().Get());
-					ASSERT_THAT(IsTrue(bAreMeshesEqual, TEXT("Meshes are not equal")));
+					const auto& GroundTruthAsset = GetGroundTruthAsset()[Type];
+					CompareToReferenceAsset(GroundTruthAsset, MeshPartitionEditorComponent);
 				});
 	}
 
@@ -361,6 +344,26 @@ TEST_CLASS_WITH_FLAGS(Water, "MeshPartition", EAutomationTestFlags::EditorContex
 
 		return Modifier;
 	}
+
+	void CompareToReferenceAsset(const FString& TestName, const UMeshPartitionEditorComponent* MeshPartitionComponent)
+	{
+		const FString ReferencePackagePath = ReferencePath / TestName;
+		const FString ReferenceAssetPath = ReferencePackagePath + "." + TestName;
+
+		if (TestUtils::ShouldUpdateReferenceAssets())
+		{
+			[[maybe_unused]] const bool bSuccess =
+				TestUtils::CreateOrUpdateReference(MeshPartitionComponent, ReferencePackagePath, { TestName }, TestName);
+		}
+
+		FString ErrorMessage;
+		bool bValid = TestUtils::ValidatePreviewMeshAgainstGroundTruth(
+			MeshPartitionComponent,
+			{ ReferenceAssetPath },
+			ErrorMessage);
+
+		ASSERT_THAT(IsTrue(bValid, *ErrorMessage));
+	}
 };
 
 #endif // WITH_EDITOR
